RAD symbol library

The RAD symbol library support many debug info formats such as pdb and elf variants. It also supports multithreaded debug info parsing and symbol lookups. Compiled libraries for Windows and headers.

diff --git a/Engine/Source/Runtime/SymsLib/SymsLib.Build.cs b/Engine/Source/Runtime/SymsLib/SymsLib.Build.cs
new file mode 100644
--- /dev/null
+++ b/Engine/Source/Runtime/SymsLib/SymsLib.Build.cs
@@ -0,0 +1,23 @@
+// Copyright Epic Games, Inc. All Rights Reserved.
+
+using UnrealBuildTool;
+
+public class SymsLib : ModuleRules
+{
+	public SymsLib(ReadOnlyTargetRules Target) : base(Target)
+	{
+		Type = ModuleType.External;
+		
+		string LibPathBase =  ModuleDirectory + "/lib";
+		string IncludePath = ModuleDirectory + "/include";
+
+		PublicSystemIncludePaths.Add(IncludePath);
+
+		if (Target.Platform.IsInGroup(UnrealPlatformGroup.Windows))
+		{
+			bool bUseDebug = Target.Configuration == UnrealTargetConfiguration.Debug;
+			string LibPath = LibPathBase + "/x64" + (bUseDebug ? "/Debug" : "/Release");
+			PublicAdditionalLibraries.Add(LibPath + "/SymsLib.lib");
+		}
+	}
+}
diff --git a/Engine/Source/Runtime/SymsLib/include/dwarf/syms_dwarf.h b/Engine/Source/Runtime/SymsLib/include/dwarf/syms_dwarf.h
new file mode 100644
--- /dev/null
+++ b/Engine/Source/Runtime/SymsLib/include/dwarf/syms_dwarf.h
@@ -0,0 +1,1922 @@
+// Copyright Epic Games, Inc. All Rights Reserved.
+#ifndef SYMS_DWARF_INCLUDE_H
+#define SYMS_DWARF_INCLUDE_H
+
+#ifndef DW_API
+#define DW_API SYMS_API
+#endif
+
+typedef syms_int dw_int;
+typedef syms_uint dw_uint;
+
+typedef SymsAddr DwOffset;
+
+/* NOTE(nick): List of all speced tags that are present in DWARF up to version 5. 
+ *
+ * Entry macro: 
+ * #define(name, value, version) 
+ *
+ */
+#define DW_TAG_LIST \
+  X(DW_TAG_ARRAY_TYPE,                0x01, DWARF_V3) \
+  X(DW_TAG_CLASS_TYPE,                0x02, DWARF_V3) \
+  X(DW_TAG_ENTRY_POINT,               0x03, DWARF_V3) \
+  X(DW_TAG_ENUMERATION_TYPE,          0x04, DWARF_V3) \
+  X(DW_TAG_FORMAL_PARAMETER,          0x05, DWARF_V3) \
+  X(DW_TAG_IMPORTED_DECLARATION,      0x08, DWARF_V3) \
+  X(DW_TAG_LABEL,                     0x0a, DWARF_V3) \
+  X(DW_TAG_LEXICAL_BLOCK,             0x0b, DWARF_V3) \
+  X(DW_TAG_MEMBER,                    0x0d, DWARF_V3) \
+  X(DW_TAG_POINTER_TYPE,              0x0f, DWARF_V3) \
+  X(DW_TAG_REFERENCE_TYPE,            0x10, DWARF_V3) \
+  X(DW_TAG_COMPILE_UNIT,              0x11, DWARF_V3) \
+  X(DW_TAG_STRING_TYPE,               0x12, DWARF_V3) \
+  X(DW_TAG_STRUCTURE_TYPE,            0x13, DWARF_V3) \
+  X(DW_TAG_SUBROUTINE_TYPE,           0x15, DWARF_V3) \
+  X(DW_TAG_TYPEDEF,                   0x16, DWARF_V3) \
+  X(DW_TAG_UNION_TYPE,                0x17, DWARF_V3) \
+  X(DW_TAG_UNSPECIFIED_PARAMETERS,    0x18, DWARF_V3) \
+  X(DW_TAG_VARIANT,                   0x19, DWARF_V3) \
+  X(DW_TAG_COMMON_BLOCK,              0x1a, DWARF_V3) \
+  X(DW_TAG_COMMON_INCLUSION,          0x1b, DWARF_V3) \
+  X(DW_TAG_INHERITANCE,               0x1c, DWARF_V3) \
+  X(DW_TAG_INLINED_SUBROUTINE,        0x1d, DWARF_V3) \
+  X(DW_TAG_MODULE,                    0x1e, DWARF_V3) \
+  X(DW_TAG_PTR_TO_MEMBER_TYPE,        0x1f, DWARF_V3) \
+  X(DW_TAG_SET_TYPE,                  0x20, DWARF_V3) \
+  X(DW_TAG_SUBRANGE_TYPE,             0x21, DWARF_V3) \
+  X(DW_TAG_WITH_STMT,                 0x22, DWARF_V3) \
+  X(DW_TAG_ACCESS_DECLARATION,        0x23, DWARF_V3) \
+  X(DW_TAG_BASE_TYPE,                 0x24, DWARF_V3) \
+  X(DW_TAG_CATCH_BLOCK,               0x25, DWARF_V3) \
+  X(DW_TAG_CONST_TYPE,                0x26, DWARF_V3) \
+  X(DW_TAG_CONSTANT,                  0x27, DWARF_V3) \
+  X(DW_TAG_ENUMERATOR,                0x28, DWARF_V3) \
+  X(DW_TAG_FILE_TYPE,                 0x29, DWARF_V3) \
+  X(DW_TAG_FRIEND,                    0x2a, DWARF_V3) \
+  X(DW_TAG_NAMELIST,                  0x2b, DWARF_V3) \
+  X(DW_TAG_NAMELIST_ITEM,             0x2c, DWARF_V3) \
+  X(DW_TAG_PACKED_TYPE,               0x2d, DWARF_V3) \
+  X(DW_TAG_SUBPROGRAM,                0x2e, DWARF_V3) \
+  X(DW_TAG_TEMPLATE_TYPE_PARAMETER,   0x2f, DWARF_V3) \
+  X(DW_TAG_TEMPLATE_VALUE_PARAMETER,  0x30, DWARF_V3) \
+  X(DW_TAG_THROWN_TYPE,               0x31, DWARF_V3) \
+  X(DW_TAG_TRY_BLOCK,                 0x32, DWARF_V3) \
+  X(DW_TAG_VARIANT_PART,              0x33, DWARF_V3) \
+  X(DW_TAG_VARIABLE,                  0x34, DWARF_V3) \
+  X(DW_TAG_VOLATILE_TYPE,             0x35, DWARF_V3) \
+  X(DW_TAG_DWARF_PROCEDURE,           0x36, DWARF_V3) \
+  X(DW_TAG_RESTRICT_TYPE,             0x37, DWARF_V3) \
+  X(DW_TAG_INTERFACE_TYPE,            0x38, DWARF_V3) \
+  X(DW_TAG_NAMESPACE,                 0x39, DWARF_V3) \
+  X(DW_TAG_IMPORTED_MODULE,           0x3a, DWARF_V3) \
+  X(DW_TAG_UNSPECIFIED_TYPE,          0x3b, DWARF_V3) \
+  X(DW_TAG_PARTIAL_UNIT,              0x3c, DWARF_V3) \
+  X(DW_TAG_IMPORTED_UNIT,             0x3d, DWARF_V3) \
+  X(DW_TAG_CONDITION,                 0x3f, DWARF_V3) \
+  X(DW_TAG_SHARED_TYPE,               0x40, DWARF_V3) \
+  X(DW_TAG_TYPE_UNIT,                 0x41, DWARF_V4) \
+  X(DW_TAG_RVALUE_REFERENCE_TYPE,     0x42, DWARF_V4) \
+  X(DW_TAG_TEMPLATE_ALIAS,            0x43, DWARF_V4) \
+  X(DW_TAG_COARRAY_TYPE,              0x44, DWARF_V5) \
+  X(DW_TAG_GENERIC_SUBRANGE,          0x45, DWARF_V5) \
+  X(DW_TAG_DYNAMIC_TYPE,              0x46, DWARF_V5) \
+  X(DW_TAG_ATOMIC_TYPE,               0x47, DWARF_V5) \
+  X(DW_TAG_CALL_SITE,                 0x48, DWARF_V5) \
+  X(DW_TAG_CALL_SITE_PARAMETER,       0x49, DWARF_V5) \
+  X(DW_TAG_SKELETON_UNIT,             0x4A, DWARF_V5) \
+  X(DW_TAG_IMMUTABLE_TYPE,            0x4B, DWARF_V5) \
+  X(DW_TAG_GNU_CALL_SITE,             0x4109, DWARF_V3) \
+  X(DW_TAG_GNU_CALL_SITE_PARAMETER,   0x410a, DWARF_V3) \
+  /* End of the list */
+
+/* NOTE(nick): List of all possible attributes that are present in DWARF up to version 5. 
+ *
+ * Entry macro:
+ * #define(name, value, version, class) 
+ *
+ */
+#define DW_ATTRIB_LIST \
+  X(DW_AT_SIBLING,                  0x01, DWARF_V3, DW_AT_CLASS_REFERENCE) \
+  X(DW_AT_LOCATION,                 0x02, DWARF_V3, DW_AT_CLASS_EXPRLOC|DW_AT_CLASS_LOCLISTPTR) \
+  X(DW_AT_NAME,                     0x03, DWARF_V3, DW_AT_CLASS_STRING) \
+  X(DW_AT_ORDERING,                 0x09, DWARF_V3, DW_AT_CLASS_CONST) \
+  X(DW_AT_BYTE_SIZE,                0x0B, DWARF_V3, DW_AT_CLASS_BLOCK|DW_AT_CLASS_CONST|DW_AT_CLASS_REFERENCE) \
+  X(DW_AT_BIT_OFFSET,               0x0C, DWARF_V3, DW_AT_CLASS_BLOCK|DW_AT_CLASS_CONST|DW_AT_CLASS_REFERENCE) \
+  X(DW_AT_BIT_SIZE,                 0x0D, DWARF_V3, DW_AT_CLASS_BLOCK|DW_AT_CLASS_CONST|DW_AT_CLASS_REFERENCE) \
+  X(DW_AT_STMT_LIST,                0x10, DWARF_V3, DW_AT_CLASS_LINEPTR) \
+  X(DW_AT_LOW_PC,                   0x11, DWARF_V3, DW_AT_CLASS_ADDRESS) \
+  X(DW_AT_HIGH_PC,                  0x12, DWARF_V3, DW_AT_CLASS_ADDRESS|DW_AT_CLASS_CONST) \
+  X(DW_AT_LANGUAGE,                 0x13, DWARF_V3, DW_AT_CLASS_CONST) \
+  X(DW_AT_DISCR,                    0x15, DWARF_V3, DW_AT_CLASS_REFERENCE) \
+  X(DW_AT_DISCR_VALUE,              0x16, DWARF_V3, DW_AT_CLASS_CONST) \
+  X(DW_AT_VISIBILITY,               0x17, DWARF_V3, DW_AT_CLASS_CONST) \
+  X(DW_AT_IMPORT,                   0x18, DWARF_V3, DW_AT_CLASS_REFERENCE) \
+  X(DW_AT_STRING_LENGTH,            0x19, DWARF_V3, DW_AT_CLASS_BLOCK|DW_AT_CLASS_LOCLISTPTR) \
+  X(DW_AT_COMMON_REFERENCE,         0x1a, DWARF_V3, DW_AT_CLASS_REFERENCE) \
+  X(DW_AT_COMP_DIR,                 0x1b, DWARF_V3, DW_AT_CLASS_STRING) \
+  X(DW_AT_CONST_VALUE,              0x1c, DWARF_V3, DW_AT_CLASS_BLOCK|DW_AT_CLASS_CONST|DW_AT_CLASS_STRING) \
+  X(DW_AT_CONTAINING_TYPE,          0x1d, DWARF_V3, DW_AT_CLASS_REFERENCE) \
+  X(DW_AT_DEFAULT_VALUE,            0x1e, DWARF_V3, DW_AT_CLASS_REFERENCE) \
+  X(DW_AT_INLINE,                   0x20, DWARF_V3, DW_AT_CLASS_CONST) \
+  X(DW_AT_IS_OPTIONAL,              0x21, DWARF_V3, DW_AT_CLASS_FLAG) \
+  X(DW_AT_LOWER_BOUND,              0x22, DWARF_V3, DW_AT_CLASS_BLOCK|DW_AT_CLASS_CONST|DW_AT_CLASS_REFERENCE) \
+  X(DW_AT_PRODUCER,                 0x25, DWARF_V3, DW_AT_CLASS_STRING) \
+  X(DW_AT_PROTOTYPED,               0x27, DWARF_V3, DW_AT_CLASS_FLAG) \
+  X(DW_AT_RETURN_ADDR,              0x2a, DWARF_V3, DW_AT_CLASS_BLOCK|DW_AT_CLASS_LOCLISTPTR) \
+  X(DW_AT_START_SCOPE,              0x2c, DWARF_V3, DW_AT_CLASS_CONST) \
+  X(DW_AT_BIT_STRIDE,               0x2e, DWARF_V3, DW_AT_CLASS_CONST) \
+  X(DW_AT_UPPER_BOUND,              0x2f, DWARF_V3, DW_AT_CLASS_BLOCK|DW_AT_CLASS_CONST|DW_AT_CLASS_REFERENCE) \
+  X(DW_AT_ABSTRACT_ORIGIN,          0x31, DWARF_V3, DW_AT_CLASS_REFERENCE) \
+  X(DW_AT_ACCESSIBILITY,            0x32, DWARF_V3, DW_AT_CLASS_CONST) \
+  X(DW_AT_ADDRESS_CLASS,            0x33, DWARF_V3, DW_AT_CLASS_CONST) \
+  X(DW_AT_ARTIFICIAL,               0x34, DWARF_V3, DW_AT_CLASS_FLAG) \
+  X(DW_AT_BASE_TYPES,               0x35, DWARF_V3, DW_AT_CLASS_REFERENCE) \
+  X(DW_AT_CALLING_CONVENTION,       0x36, DWARF_V3, DW_AT_CLASS_CONST) \
+  X(DW_AT_COUNT,                    0x37, DWARF_V3, DW_AT_CLASS_BLOCK|DW_AT_CLASS_CONST|DW_AT_CLASS_REFERENCE) \
+  X(DW_AT_DATA_MEMBER_LOCATION,     0x38, DWARF_V3, DW_AT_CLASS_BLOCK|DW_AT_CLASS_CONST|DW_AT_CLASS_LOCLISTPTR) \
+  X(DW_AT_DECL_COLUMN,              0x39, DWARF_V3, DW_AT_CLASS_CONST) \
+  X(DW_AT_DECL_FILE,                0x3a, DWARF_V3, DW_AT_CLASS_CONST) \
+  X(DW_AT_DECL_LINE,                0x3b, DWARF_V3, DW_AT_CLASS_CONST) \
+  X(DW_AT_DECLARATION,              0x3c, DWARF_V3, DW_AT_CLASS_FLAG) \
+  X(DW_AT_DISCR_LIST,               0x3d, DWARF_V3, DW_AT_CLASS_BLOCK) \
+  X(DW_AT_ENCODING,                 0x3e, DWARF_V3, DW_AT_CLASS_CONST) \
+  X(DW_AT_EXTERNAL,                 0x3f, DWARF_V3, DW_AT_CLASS_FLAG) \
+  X(DW_AT_FRAME_BASE,               0x40, DWARF_V3, DW_AT_CLASS_EXPRLOC|DW_AT_CLASS_LOCLISTPTR) \
+  X(DW_AT_FRIEND,                   0x41, DWARF_V3, DW_AT_CLASS_REFERENCE) \
+  X(DW_AT_IDENTIFIER_CASE,          0x42, DWARF_V3, DW_AT_CLASS_CONST) \
+  X(DW_AT_MACRO_INFO,               0x43, DWARF_V3, DW_AT_CLASS_MACPTR) \
+  X(DW_AT_NAMELIST_ITEM,            0x44, DWARF_V3, DW_AT_CLASS_BLOCK) \
+  X(DW_AT_PRIORITY,                 0x45, DWARF_V3, DW_AT_CLASS_REFERENCE) \
+  X(DW_AT_SEGMENT,                  0x46, DWARF_V3, DW_AT_CLASS_BLOCK|DW_AT_CLASS_LOCLISTPTR) \
+  X(DW_AT_SPECIFICATION,            0x47, DWARF_V3, DW_AT_CLASS_REFERENCE) \
+  X(DW_AT_STATIC_LINK,              0x48, DWARF_V3, DW_AT_CLASS_BLOCK|DW_AT_CLASS_LOCLISTPTR) \
+  X(DW_AT_TYPE,                     0x49, DWARF_V3, DW_AT_CLASS_REFERENCE) \
+  X(DW_AT_USE_LOCATION,             0x4a, DWARF_V3, DW_AT_CLASS_BLOCK|DW_AT_CLASS_LOCLISTPTR) \
+  X(DW_AT_VARIABLE_PARAMETER,       0x4b, DWARF_V3, DW_AT_CLASS_FLAG) \
+  X(DW_AT_VIRTUALITY,               0x4c, DWARF_V3, DW_AT_CLASS_CONST) \
+  X(DW_AT_VTABLE_ELEM_LOCATION,     0x4d, DWARF_V3, DW_AT_CLASS_BLOCK|DW_AT_CLASS_LOCLISTPTR) \
+  X(DW_AT_ALLOCATED,                0x4e, DWARF_V3, DW_AT_CLASS_BLOCK|DW_AT_CLASS_CONST|DW_AT_CLASS_REFERENCE) \
+  X(DW_AT_ASSOCIATED,               0x4f, DWARF_V3, DW_AT_CLASS_BLOCK|DW_AT_CLASS_CONST|DW_AT_CLASS_REFERENCE) \
+  X(DW_AT_DATA_LOCATION,            0x50, DWARF_V3, DW_AT_CLASS_BLOCK) \
+  X(DW_AT_BYTE_STRIDE,              0x51, DWARF_V3, DW_AT_CLASS_BLOCK|DW_AT_CLASS_CONST|DW_AT_CLASS_REFERENCE) \
+  X(DW_AT_ENTRY_PC,                 0x52, DWARF_V3, DW_AT_CLASS_ADDRESS) \
+  X(DW_AT_USE_UTF8,                 0x53, DWARF_V3, DW_AT_CLASS_FLAG) \
+  X(DW_AT_EXTENSION,                0x54, DWARF_V3, DW_AT_CLASS_REFERENCE) \
+  X(DW_AT_RANGES,                   0x55, DWARF_V3, DW_AT_CLASS_RNGLISTPTR) \
+  X(DW_AT_TRAMPOLINE,               0x56, DWARF_V3, DW_AT_CLASS_ADDRESS|DW_AT_CLASS_FLAG|DW_AT_CLASS_REFERENCE|DW_AT_CLASS_STRING) \
+  X(DW_AT_CALL_COLUMN,              0x57, DWARF_V3, DW_AT_CLASS_CONST) \
+  X(DW_AT_CALL_FILE,                0x58, DWARF_V3, DW_AT_CLASS_CONST) \
+  X(DW_AT_CALL_LINE,                0x59, DWARF_V3, DW_AT_CLASS_CONST) \
+  X(DW_AT_DESCRIPTION,              0x5a, DWARF_V3, DW_AT_CLASS_STRING) \
+  X(DW_AT_BINARY_SCALE,             0x5b, DWARF_V3, DW_AT_CLASS_CONST) \
+  X(DW_AT_DECIMAL_SCALE,            0x5c, DWARF_V3, DW_AT_CLASS_CONST) \
+  X(DW_AT_SMALL,                    0x5d, DWARF_V3, DW_AT_CLASS_REFERENCE) \
+  X(DW_AT_DECIMAL_SIGN,             0x5e, DWARF_V3, DW_AT_CLASS_CONST) \
+  X(DW_AT_DIGIT_COUNT,              0x5f, DWARF_V3, DW_AT_CLASS_CONST) \
+  X(DW_AT_PICTURE_STRING,           0x60, DWARF_V3, DW_AT_CLASS_STRING) \
+  X(DW_AT_MUTABLE,                  0x61, DWARF_V3, DW_AT_CLASS_FLAG) \
+  X(DW_AT_THREADS_SCALED,           0x62, DWARF_V3, DW_AT_CLASS_FLAG) \
+  X(DW_AT_EXPLICIT,                 0x63, DWARF_V3, DW_AT_CLASS_FLAG) \
+  X(DW_AT_OBJECT_POINTER,           0x64, DWARF_V3, DW_AT_CLASS_REFERENCE) \
+  X(DW_AT_ENDIANITY,                0x65, DWARF_V3, DW_AT_CLASS_CONST) \
+  X(DW_AT_ELEMENTAL,                0x66, DWARF_V3, DW_AT_CLASS_FLAG) \
+  X(DW_AT_PURE,                     0x67, DWARF_V3, DW_AT_CLASS_FLAG) \
+  X(DW_AT_RECURSIVE,                0x68, DWARF_V3, DW_AT_CLASS_FLAG) \
+  X(DW_AT_SIGNATURE,                0x69, DWARF_V4, DW_AT_CLASS_REFERENCE) \
+  X(DW_AT_MAIN_SUBPROGRAM,          0x6a, DWARF_V4, DW_AT_CLASS_FLAG) \
+  X(DW_AT_DATA_BIT_OFFSET,          0x6b, DWARF_V4, DW_AT_CLASS_CONST) \
+  X(DW_AT_CONST_EXPR,               0x6c, DWARF_V4, DW_AT_CLASS_FLAG) \
+  X(DW_AT_ENUM_CLASS,               0x6d, DWARF_V4, DW_AT_CLASS_FLAG) \
+  X(DW_AT_LINKAGE_NAME,             0x6e, DWARF_V4, DW_AT_CLASS_STRING) \
+  X(DW_AT_STRING_LENGTH_BIT_SIZE,   0x6f, DWARF_V5, DW_AT_CLASS_CONST) \
+  X(DW_AT_STRING_LENGTH_BYTE_SIZE,  0x70, DWARF_V5, DW_AT_CLASS_CONST) \
+  X(DW_AT_RANK,                     0x71, DWARF_V5, DW_AT_CLASS_CONST|DW_AT_CLASS_EXPRLOC) \
+  X(DW_AT_STR_OFFSETS_BASE,         0x72, DWARF_V5, DW_AT_CLASS_STROFFSETSPTR) \
+  X(DW_AT_ADDR_BASE,                0x73, DWARF_V5, DW_AT_CLASS_ADDRPTR) \
+  X(DW_AT_RNGLISTS_BASE,            0x74, DWARF_V5, DW_AT_CLASS_RNGLISTPTR) \
+  X(DW_AT_DWO_NAME,                 0x76, DWARF_V5, DW_AT_CLASS_STRING) \
+  X(DW_AT_REFERENCE,                0x77, DWARF_V5, DW_AT_CLASS_FLAG) \
+  X(DW_AT_RVALUE_REFERENCE,         0x78, DWARF_V5, DW_AT_CLASS_FLAG) \
+  X(DW_AT_MACROS,                   0x79, DWARF_V5, DW_AT_CLASS_MACPTR) \
+  X(DW_AT_CALL_ALL_CALLS,           0x7a, DWARF_V5, DW_AT_CLASS_FLAG) \
+  X(DW_AT_CALL_ALL_SOURCE_CALLS,    0x7b, DWARF_V5, DW_AT_CLASS_FLAG) \
+  X(DW_AT_CALL_ALL_TAIL_CALLS,      0x7c, DWARF_V5, DW_AT_CLASS_FLAG) \
+  X(DW_AT_CALL_RETURN_PC,           0x7d, DWARF_V5, DW_AT_CLASS_ADDRESS) \
+  X(DW_AT_CALL_VALUE,               0x7e, DWARF_V5, DW_AT_CLASS_EXPRLOC) \
+  X(DW_AT_CALL_ORIGIN,              0x7f, DWARF_V5, DW_AT_CLASS_EXPRLOC) \
+  X(DW_AT_CALL_PARAMETER,           0x80, DWARF_V5, DW_AT_CLASS_REFERENCE) \
+  X(DW_AT_CALL_PC,                  0x81, DWARF_V5, DW_AT_CLASS_ADDRESS) \
+  X(DW_AT_CALL_TAIL_CALL,           0x82, DWARF_V5, DW_AT_CLASS_FLAG) \
+  X(DW_AT_CALL_TARGET,              0x83, DWARF_V5, DW_AT_CLASS_EXPRLOC) \
+  X(DW_AT_CALL_TARGET_CLOBBERED,    0x84, DWARF_V5, DW_AT_CLASS_EXPRLOC) \
+  X(DW_AT_CALL_DATA_LOCATION,       0x85, DWARF_V5, DW_AT_CLASS_EXPRLOC) \
+  X(DW_AT_CALL_DATA_VALUE,          0x86, DWARF_V5, DW_AT_CLASS_EXPRLOC) \
+  X(DW_AT_NORETURN,                 0x87, DWARF_V5, DW_AT_CLASS_FLAG) \
+  X(DW_AT_ALIGNMENT,                0x88, DWARF_V5, DW_AT_CLASS_CONST) \
+  X(DW_AT_EXPORT_SYMBOLS,           0x89, DWARF_V5, DW_AT_CLASS_FLAG) \
+  X(DW_AT_DELETED,                  0x8a, DWARF_V5, DW_AT_CLASS_FLAG) \
+  X(DW_AT_DEFAULTED,                0x8b, DWARF_V5, DW_AT_CLASS_CONST) \
+  X(DW_AT_LOCLISTS_BASE,            0x8c, DWARF_V5, DW_AT_CLASS_LOCLISTPTR) \
+  X(DW_AT_GNU_VECTOR,                     0x2107, DWARF_V3, DW_AT_CLASS_FLAG) \
+  X(DW_AT_GNU_GUARDED_BY,                 0x2108, DWARF_V2, DW_AT_CLASS_UNDEFINED) \
+  X(DW_AT_GNU_PT_GUARDED_BY,              0x2109, DWARF_V2, DW_AT_CLASS_UNDEFINED) \
+  X(DW_AT_GNU_GUARDED,                    0x210a, DWARF_V2, DW_AT_CLASS_UNDEFINED) \
+  X(DW_AT_GNU_PT_GUARDED,                 0x210b, DWARF_V2, DW_AT_CLASS_UNDEFINED) \
+  X(DW_AT_GNU_LOCKS_EXCLUDED,             0x210c, DWARF_V2, DW_AT_CLASS_UNDEFINED) \
+  X(DW_AT_GNU_EXCLUSIVE_LOCKS_REQUIRED,   0x210d, DWARF_V2, DW_AT_CLASS_UNDEFINED) \
+  X(DW_AT_GNU_SHARED_LOCKS_REQUIRED,      0x210e, DWARF_V2, DW_AT_CLASS_UNDEFINED) \
+  X(DW_AT_GNU_ODR_SIGNATURE,              0x210f, DWARF_V2, DW_AT_CLASS_UNDEFINED) \
+  X(DW_AT_GNU_TEMPLATE_NAME,              0x2110, DWARF_V3, DW_AT_CLASS_UNDEFINED) \
+  X(DW_AT_GNU_CALL_SITE_VALUE,            0x2111, DWARF_V3, DW_AT_CLASS_EXPRLOC) \
+  X(DW_AT_GNU_CALL_SITE_DATA_VALUE,       0x2112, DWARF_V3, DW_AT_CLASS_EXPRLOC) \
+  X(DW_AT_GNU_CALL_SITE_TARGET,           0x2113, DWARF_V3, DW_AT_CLASS_EXPRLOC) \
+  X(DW_AT_GNU_CALL_SITE_TARGET_CLOBBERED, 0x2114, DWARF_V3, DW_AT_CLASS_EXPRLOC) \
+  X(DW_AT_GNU_TAIL_CALL,                  0x2115, DWARF_V3, DW_AT_CLASS_FLAG) \
+  X(DW_AT_GNU_ALL_TAIL_CALL_SITES,        0x2116, DWARF_V3, DW_AT_CLASS_FLAG) \
+  X(DW_AT_GNU_ALL_CALL_SITES,             0x2117, DWARF_V3, DW_AT_CLASS_FLAG) \
+  X(DW_AT_GNU_ALL_SOURCE_CALL_SITES,      0x2118, DWARF_V3, DW_AT_CLASS_FLAG) \
+  X(DW_AT_GNU_MACROS,                     0x2119, DWARF_V2, DW_AT_CLASS_FLAG) \
+  X(DW_AT_GNU_DELETED,                    0x211a, DWARF_V2, DW_AT_CLASS_UNDEFINED) \
+  X(DW_AT_GNU_DWO_NAME,                   0x2130, DWARF_V2, DW_AT_CLASS_STRING) \
+  X(DW_AT_GNU_DWO_ID,                     0x2131, DWARF_V2, DW_AT_CLASS_CONST) \
+  X(DW_AT_GNU_RANGES_BASE,                0x2132, DWARF_V2, DW_AT_CLASS_UNDEFINED) \
+  X(DW_AT_GNU_ADDR_BASE,                  0x2133, DWARF_V2, DW_AT_CLASS_ADDRPTR) \
+  X(DW_AT_GNU_PUBNAMES,                   0x2134, DWARF_V2, DW_AT_CLASS_FLAG) \
+  X(DW_AT_GNU_PUBTYPES,                   0x2135, DWARF_V2, DW_AT_CLASS_UNDEFINED) \
+  X(DW_AT_GNU_DISCRIMINATOR,              0x2136, DWARF_V3, DW_AT_CLASS_CONST) \
+  X(DW_AT_GNU_LOCVIEWS,                   0x2137, DWARF_V2, DW_AT_CLASS_UNDEFINED)\
+  X(DW_AT_GNU_ENTRY_VIEW,                 0x2138, DWARF_V2, DW_AT_CLASS_UNDEFINED) \
+  X(DW_AT_VMS_RTNBEG_PD_ADDRESS,          0x2201, DWARF_V3, DW_AT_CLASS_UNDEFINED) \
+  X(DW_AT_USE_GNAT_DESCRIPTIVE_TYPE,      0x2301, DWARF_V2, DW_AT_CLASS_UNDEFINED) \
+  X(DW_AT_GNAT_DESCRIPTIVE_TYPE,          0x2302, DWARF_V2, DW_AT_CLASS_UNDEFINED) \
+  X(DW_AT_GNU_NUMERATOR,                  0x2303, DWARF_V2, DW_AT_CLASS_UNDEFINED) \
+  X(DW_AT_GNU_DENOMINATOR,                0x2304, DWARF_V2, DW_AT_CLASS_UNDEFINED) \
+  X(DW_AT_GNU_BIAS,                       0x2305, DWARF_V2, DW_AT_CLASS_UNDEFINED) \
+  X(DW_AT_UPC_THREADS_SCALED,             0x3210, DWARF_V2, DW_AT_CLASS_UNDEFINED) \
+  X(DW_AT_PGI_LBASE,                      0x3a00, DWARF_V2, DW_AT_CLASS_UNDEFINED) \
+  X(DW_AT_PGI_SOFFSET,                    0x3a01, DWARF_V2, DW_AT_CLASS_UNDEFINED) \
+  X(DW_AT_PGI_LSTRIDE,                    0x3a02, DWARF_V2, DW_AT_CLASS_UNDEFINED) \
+  X(DW_AT_APPLE_OPTIMIZED,            0x3fe1, DWARF_V3, DW_AT_CLASS_FLAG) \
+  X(DW_AT_APPLE_FLAGS,                0x3fe2, DWARF_V3, DW_AT_CLASS_FLAG) \
+  X(DW_AT_APPLE_ISA,                  0x3fe3, DWARF_V3, DW_AT_CLASS_FLAG) \
+  X(DW_AT_APPLE_BLOCK,                0x3fe4, DWARF_V3, DW_AT_CLASS_UNDEFINED) \
+  X(DW_AT_APPLE_MAJOR_RUNTIME_VERS,   0x3fe5, DWARF_V3, DW_AT_CLASS_UNDEFINED) \
+  X(DW_AT_APPLE_RUNTIME_CLASS,        0x3fe6, DWARF_V3, DW_AT_CLASS_UNDEFINED) \
+  X(DW_AT_APPLE_OMIT_FRAME_PTR,       0x3fe7, DWARF_V3, DW_AT_CLASS_FLAG) \
+  X(DW_AT_APPLE_PROPERTY_NAME,        0x3fe8, DWARF_V3, DW_AT_CLASS_UNDEFINED) \
+  X(DW_AT_APPLE_PROPERTY_GETTER,      0x3fe9, DWARF_V3, DW_AT_CLASS_UNDEFINED) \
+  X(DW_AT_APPLE_PROPERTY_SETTER,      0x3fea, DWARF_V3, DW_AT_CLASS_UNDEFINED) \
+  X(DW_AT_APPLE_PROPERTY_ATTRIBUTE,   0x3feb, DWARF_V3, DW_AT_CLASS_UNDEFINED) \
+  X(DW_AT_APPLE_OBJC_COMPLETE_TYPE,   0x3fec, DWARF_V3, DW_AT_CLASS_UNDEFINED) \
+  X(DW_AT_APPLE_PROPERTY,             0x3fed, DWARF_V3, DW_AT_CLASS_UNDEFINED) \
+  /* End of the list */
+
+#define DW_FORM_LIST \
+  X(DW_FORM_ADDR,             0x01, DWARF_V3, DW_AT_CLASS_ADDRESS) \
+  X(DW_FORM_BLOCK2,           0x03, DWARF_V3, DW_AT_CLASS_BLOCK) \
+  X(DW_FORM_BLOCK4,           0x04, DWARF_V3, DW_AT_CLASS_BLOCK) \
+  X(DW_FORM_DATA2,            0x05, DWARF_V3, DW_AT_CLASS_CONST) \
+  X(DW_FORM_DATA4,            0x06, DWARF_V3, DW_AT_CLASS_CONST|DW_AT_CLASS_LINEPTR|DW_AT_CLASS_LOCLISTPTR|DW_AT_CLASS_MACPTR|DW_AT_CLASS_RNGLISTPTR) \
+  X(DW_FORM_DATA8,            0x07, DWARF_V3, DW_AT_CLASS_CONST|DW_AT_CLASS_LINEPTR|DW_AT_CLASS_LOCLISTPTR|DW_AT_CLASS_MACPTR|DW_AT_CLASS_RNGLISTPTR) \
+  X(DW_FORM_STRING,           0x08, DWARF_V3, DW_AT_CLASS_STRING) \
+  X(DW_FORM_BLOCK,            0x09, DWARF_V3, DW_AT_CLASS_BLOCK) \
+  X(DW_FORM_BLOCK1,           0x0a, DWARF_V3, DW_AT_CLASS_BLOCK) \
+  X(DW_FORM_DATA1,            0x0b, DWARF_V3, DW_AT_CLASS_CONST) \
+  X(DW_FORM_FLAG,             0x0c, DWARF_V3, DW_AT_CLASS_FLAG) \
+  X(DW_FORM_SDATA,            0x0d, DWARF_V3, DW_AT_CLASS_CONST) \
+  X(DW_FORM_STRP,             0x0e, DWARF_V3, DW_AT_CLASS_STRING) \
+  X(DW_FORM_UDATA,            0x0f, DWARF_V3, DW_AT_CLASS_CONST) \
+  X(DW_FORM_REF_ADDR,         0x10, DWARF_V3, DW_AT_CLASS_REFERENCE) \
+  X(DW_FORM_REF1,             0x11, DWARF_V3, DW_AT_CLASS_REFERENCE) \
+  X(DW_FORM_REF2,             0x12, DWARF_V3, DW_AT_CLASS_REFERENCE) \
+  X(DW_FORM_REF4,             0x13, DWARF_V3, DW_AT_CLASS_REFERENCE) \
+  X(DW_FORM_REF8,             0x14, DWARF_V3, DW_AT_CLASS_REFERENCE) \
+  X(DW_FORM_REF_UDATA,        0x15, DWARF_V3, DW_AT_CLASS_REFERENCE) \
+  X(DW_FORM_INDIRECT,         0x16, DWARF_V3, 0) \
+  X(DW_FORM_SEC_OFFSET,       0x17, DWARF_V4, DW_AT_CLASS_LINEPTR|DW_AT_CLASS_LOCLISTPTR|DW_AT_CLASS_MACPTR|DW_AT_CLASS_RNGLISTPTR) \
+  X(DW_FORM_EXPRLOC,          0x18, DWARF_V4, DW_AT_CLASS_EXPRLOC) \
+  X(DW_FORM_FLAG_PRESENT,     0x19, DWARF_V4, DW_AT_CLASS_FLAG) \
+  X(DW_FORM_REF_SIG8,         0x20, DWARF_V4, DW_AT_CLASS_REFERENCE) \
+  X(DW_FORM_STRX,             0x1a, DWARF_V5, DW_AT_CLASS_STRING) \
+  X(DW_FORM_ADDRX,            0x1b, DWARF_V5, DW_AT_CLASS_ADDRESS) \
+  X(DW_FORM_REF_SUP4,         0x1c, DWARF_V5, DW_AT_CLASS_REFERENCE) \
+  X(DW_FORM_STRP_SUP,         0x1d, DWARF_V5, DW_AT_CLASS_STRING) \
+  X(DW_FORM_DATA16,           0x1e, DWARF_V5, DW_AT_CLASS_CONST) \
+  X(DW_FORM_LINE_STRP,        0x1f, DWARF_V5, DW_AT_CLASS_STRING) \
+  X(DW_FORM_IMPLICIT_CONST,   0x21, DWARF_V5, DW_AT_CLASS_CONST) \
+  X(DW_FORM_LOCLISTX,         0x22, DWARF_V5, DW_AT_CLASS_LOCLISTPTR) \
+  X(DW_FORM_RNGLISTX,         0x23, DWARF_V5, DW_AT_CLASS_RNGLISTPTR) \
+  X(DW_FORM_REF_SUP8,         0x24, DWARF_V5, DW_AT_CLASS_REFERENCE) \
+  X(DW_FORM_STRX1,            0x25, DWARF_V5, DW_AT_CLASS_STRING) \
+  X(DW_FORM_STRX2,            0x26, DWARF_V5, DW_AT_CLASS_STRING) \
+  X(DW_FORM_STRX3,            0x27, DWARF_V5, DW_AT_CLASS_STRING) \
+  X(DW_FORM_STRX4,            0x28, DWARF_V5, DW_AT_CLASS_STRING) \
+  X(DW_FORM_ADDRX1,           0x29, DWARF_V5, DW_AT_CLASS_ADDRESS) \
+  X(DW_FORM_ADDRX2,           0x2a, DWARF_V5, DW_AT_CLASS_ADDRESS) \
+  X(DW_FORM_ADDRX3,           0x2b, DWARF_V5, DW_AT_CLASS_ADDRESS) \
+  X(DW_FORM_ADDRX4,           0x2c, DWARF_V5, DW_AT_CLASS_ADDRESS) \
+  /* End of the list */
+
+/* NOTE(nick): List of attribute encodings that are present in DWARF up to version 5.
+ *
+ * Entry macro:
+ * #define(name, value, version) 
+ *
+ */
+#define DW_ATE_LIST \
+  X(DW_ATE_ADDRESS,           0X01, DWARF_V3) \
+  X(DW_ATE_BOOLEAN,           0X02, DWARF_V3) \
+  X(DW_ATE_COMPLEX_FLOAT,     0X03, DWARF_V3) \
+  X(DW_ATE_FLOAT,             0X04, DWARF_V3) \
+  X(DW_ATE_SIGNED,            0X05, DWARF_V3) \
+  X(DW_ATE_SIGNED_CHAR,       0X06, DWARF_V3) \
+  X(DW_ATE_UNSIGNED,          0X07, DWARF_V3) \
+  X(DW_ATE_UNSIGNED_CHAR,     0X08, DWARF_V3) \
+  X(DW_ATE_IMAGINARY_FLOAT,   0X09, DWARF_V3) \
+  X(DW_ATE_PACKED_DECIMAL,    0X0A, DWARF_V3) \
+  X(DW_ATE_NUMERIC_STRING,    0X0B, DWARF_V3) \
+  X(DW_ATE_EDITED,            0X0C, DWARF_V3) \
+  X(DW_ATE_SIGNED_FIXED,      0X0D, DWARF_V3) \
+  X(DW_ATE_UNSIGNED_FIXED,    0X0E, DWARF_V3) \
+  X(DW_ATE_DECIMAL_FLOAT,     0X0F, DWARF_V3) \
+  X(DW_ATE_UTF,               0X10, DWARF_V4) \
+  X(DW_ATE_UCS,               0X11, DWARF_V5) \
+  X(DW_ATE_ASCII,             0X12, DWARF_V5) \
+  /* END OF THE LIST */
+
+/* NOTE(NICK): LIST OF ALL EXPRESSIOn ops that are present in DWARF up to version 5. 
+ *
+ * EntRY MACRO:
+ * #deFINE(NAME, VALUE, NUM_OPERANDS, version)
+ *
+ * */
+#define DW_OP_LIST \
+  X(DW_OP_ADDR, 0x03, 1, DWARF_V3)           X(DW_OP_DEREF, 0x06, 0, DWARF_V3) \
+  X(DW_OP_CONST1U, 0x08, 1, DWARF_V3)        X(DW_OP_CONST1S, 0x09, 1, DWARF_V3) \
+  X(DW_OP_CONST2U, 0x0a, 1, DWARF_V3)        X(DW_OP_CONST2S, 0x0b, 1, DWARF_V3) \
+  X(DW_OP_CONST4U, 0x0c, 1, DWARF_V3)        X(DW_OP_CONST4S, 0x0d, 1, DWARF_V3) \
+  X(DW_OP_CONST8U, 0x0e, 1, DWARF_V3)        X(DW_OP_CONST8S, 0x0f, 1, DWARF_V3) \
+  X(DW_OP_CONSTU, 0x10, 1, DWARF_V3)         X(DW_OP_CONSTS, 0x11, 1, DWARF_V3) \
+  X(DW_OP_DUP, 0x12, 0, DWARF_V3)            X(DW_OP_DROP, 0x13, 0, DWARF_V3) \
+  X(DW_OP_OVER, 0x14, 0, DWARF_V3)           X(DW_OP_PICK, 0x15, 1, DWARF_V3) \
+  X(DW_OP_SWAP, 0x16, 0, DWARF_V3)           X(DW_OP_ROT, 0x17, 0, DWARF_V3) \
+  X(DW_OP_XDEREF, 0x18, 0, DWARF_V3)         X(DW_OP_ABS, 0x19, 0, DWARF_V3) \
+  X(DW_OP_AND, 0x1a, 0, DWARF_V3)            X(DW_OP_DIV, 0x1b, 0, DWARF_V3) \
+  X(DW_OP_MINUS, 0x1c, 0, DWARF_V3)          X(DW_OP_MOD, 0x1d, 0, DWARF_V3) \
+  X(DW_OP_MUL, 0x1e, 0, DWARF_V3)            X(DW_OP_NEG, 0x1f, 0, DWARF_V3) \
+  X(DW_OP_NOT, 0x20, 0, DWARF_V3)            X(DW_OP_OR, 0x21, 0, DWARF_V3) \
+  X(DW_OP_PLUS, 0x22, 0, DWARF_V3)           X(DW_OP_PLUS_UCONST, 0x23, 1, DWARF_V3) \
+  X(DW_OP_SHL, 0x24, 0, DWARF_V3)            X(DW_OP_SHR, 0x25, 0, DWARF_V3) \
+  X(DW_OP_SHRA, 0x26, 0, DWARF_V3)           X(DW_OP_XOR, 0x27, 0, DWARF_V3) \
+  X(DW_OP_SKIP, 0x2f, 1, DWARF_V3)           X(DW_OP_BRA, 0x28, 1, DWARF_V3) \
+  X(DW_OP_EQ, 0x29, 0, DWARF_V3)             X(DW_OP_GE, 0x2a, 0, DWARF_V3) \
+  X(DW_OP_GT, 0x2b, 0, DWARF_V3)             X(DW_OP_LE, 0x2c, 0, DWARF_V3) \
+  X(DW_OP_LT, 0x2d, 0, DWARF_V3)             X(DW_OP_NE, 0x2e, 0, DWARF_V3) \
+  X(DW_OP_LIT0, 0x30, 0, DWARF_V3)           X(DW_OP_LIT1, 0x31, 0, DWARF_V3) \
+  X(DW_OP_LIT2, 0x32, 0, DWARF_V3)           X(DW_OP_LIT3, 0x33, 0, DWARF_V3) \
+  X(DW_OP_LIT4, 0x34, 0, DWARF_V3)           X(DW_OP_LIT5, 0x35, 0, DWARF_V3) \
+  X(DW_OP_LIT6, 0x36, 0, DWARF_V3)           X(DW_OP_LIT7, 0x37, 0, DWARF_V3) \
+  X(DW_OP_LIT8, 0x38, 0, DWARF_V3)           X(DW_OP_LIT9, 0x39, 0, DWARF_V3) \
+  X(DW_OP_LIT10, 0x3a, 0, DWARF_V3)          X(DW_OP_LIT11, 0x3b, 0, DWARF_V3) \
+  X(DW_OP_LIT12, 0x3c, 0, DWARF_V3)          X(DW_OP_LIT13, 0x3d, 0, DWARF_V3) \
+  X(DW_OP_LIT14, 0x3e, 0, DWARF_V3)          X(DW_OP_LIT15, 0x3f, 0, DWARF_V3) \
+  X(DW_OP_LIT16, 0x40, 0, DWARF_V3)          X(DW_OP_LIT17, 0x41, 0, DWARF_V3) \
+  X(DW_OP_LIT18, 0x42, 0, DWARF_V3)          X(DW_OP_LIT19, 0x43, 0, DWARF_V3) \
+  X(DW_OP_LIT20, 0x44, 0, DWARF_V3)          X(DW_OP_LIT21, 0x45, 0, DWARF_V3) \
+  X(DW_OP_LIT22, 0x46, 0, DWARF_V3)          X(DW_OP_LIT23, 0x47, 0, DWARF_V3) \
+  X(DW_OP_LIT24, 0x48, 0, DWARF_V3)          X(DW_OP_LIT25, 0x49, 0, DWARF_V3) \
+  X(DW_OP_LIT26, 0x4a, 0, DWARF_V3)          X(DW_OP_LIT27, 0x4b, 0, DWARF_V3) \
+  X(DW_OP_LIT28, 0x4c, 0, DWARF_V3)          X(DW_OP_LIT29, 0x4d, 0, DWARF_V3) \
+  X(DW_OP_LIT30, 0x4e, 0, DWARF_V3)          X(DW_OP_LIT31, 0x4f, 0, DWARF_V3) \
+  X(DW_OP_REG0, 0x50, 0, DWARF_V3)           X(DW_OP_REG1, 0x51, 0, DWARF_V3) \
+  X(DW_OP_REG2, 0x52, 0, DWARF_V3)           X(DW_OP_REG3, 0x53, 0, DWARF_V3) \
+  X(DW_OP_REG4, 0x54, 0, DWARF_V3)           X(DW_OP_REG5, 0x55, 0, DWARF_V3) \
+  X(DW_OP_REG6, 0x56, 0, DWARF_V3)           X(DW_OP_REG7, 0x57, 0, DWARF_V3) \
+  X(DW_OP_REG8, 0x58, 0, DWARF_V3)           X(DW_OP_REG9, 0x59, 0, DWARF_V3) \
+  X(DW_OP_REG10, 0x5a, 0, DWARF_V3)          X(DW_OP_REG11, 0x5b, 0, DWARF_V3) \
+  X(DW_OP_REG12, 0x5c, 0, DWARF_V3)          X(DW_OP_REG13, 0x5d, 0, DWARF_V3) \
+  X(DW_OP_REG14, 0x5e, 0, DWARF_V3)          X(DW_OP_REG15, 0x5f, 0, DWARF_V3) \
+  X(DW_OP_REG16, 0x60, 0, DWARF_V3)          X(DW_OP_REG17, 0x61, 0, DWARF_V3) \
+  X(DW_OP_REG18, 0x62, 0, DWARF_V3)          X(DW_OP_REG19, 0x63, 0, DWARF_V3) \
+  X(DW_OP_REG20, 0x64, 0, DWARF_V3)          X(DW_OP_REG21, 0x65, 0, DWARF_V3) \
+  X(DW_OP_REG22, 0x66, 0, DWARF_V3)          X(DW_OP_REG23, 0x67, 0, DWARF_V3) \
+  X(DW_OP_REG24, 0x68, 0, DWARF_V3)          X(DW_OP_REG25, 0x69, 0, DWARF_V3) \
+  X(DW_OP_REG26, 0x6a, 0, DWARF_V3)          X(DW_OP_REG27, 0x6b, 0, DWARF_V3) \
+  X(DW_OP_REG28, 0x6c, 0, DWARF_V3)          X(DW_OP_REG29, 0x6d, 0, DWARF_V3) \
+  X(DW_OP_REG30, 0x6e, 0, DWARF_V3)          X(DW_OP_REG31, 0x6f, 0, DWARF_V3) \
+  X(DW_OP_BREG0, 0x70, 0, DWARF_V3)          X(DW_OP_BREG1, 0x71, 0, DWARF_V3) \
+  X(DW_OP_BREG2, 0x72, 0, DWARF_V3)          X(DW_OP_BREG3, 0x73, 0, DWARF_V3) \
+  X(DW_OP_BREG4, 0x74, 0, DWARF_V3)          X(DW_OP_BREG5, 0x75, 0, DWARF_V3) \
+  X(DW_OP_BREG6, 0x76, 0, DWARF_V3)          X(DW_OP_BREG7, 0x77, 0, DWARF_V3) \
+  X(DW_OP_BREG8, 0x78, 0, DWARF_V3)          X(DW_OP_BREG9, 0x79, 0, DWARF_V3) \
+  X(DW_OP_BREG10, 0x7a, 0, DWARF_V3)         X(DW_OP_BREG11, 0x7b, 0, DWARF_V3) \
+  X(DW_OP_BREG12, 0x7c, 0, DWARF_V3)         X(DW_OP_BREG13, 0x7d, 0, DWARF_V3) \
+  X(DW_OP_BREG14, 0x7e, 0, DWARF_V3)         X(DW_OP_BREG15, 0x7f, 0, DWARF_V3) \
+  X(DW_OP_BREG16, 0x80, 0, DWARF_V3)         X(DW_OP_BREG17, 0x81, 0, DWARF_V3) \
+  X(DW_OP_BREG18, 0x82, 0, DWARF_V3)         X(DW_OP_BREG19, 0x83, 0, DWARF_V3) \
+  X(DW_OP_BREG20, 0x84, 0, DWARF_V3)         X(DW_OP_BREG21, 0x85, 0, DWARF_V3) \
+  X(DW_OP_BREG22, 0x86, 0, DWARF_V3)         X(DW_OP_BREG23, 0x87, 0, DWARF_V3) \
+  X(DW_OP_BREG24, 0x88, 0, DWARF_V3)         X(DW_OP_BREG25, 0x89, 0, DWARF_V3) \
+  X(DW_OP_BREG26, 0x8a, 0, DWARF_V3)         X(DW_OP_BREG27, 0x8b, 0, DWARF_V3) \
+  X(DW_OP_BREG28, 0x8c, 0, DWARF_V3)         X(DW_OP_BREG29, 0x8d, 0, DWARF_V3) \
+  X(DW_OP_BREG30, 0x8e, 0, DWARF_V3)         X(DW_OP_BREG31, 0x8f, 0, DWARF_V3) \
+  X(DW_OP_REGX, 0x90, 1, DWARF_V3)           X(DW_OP_FBREG, 0x91, 1, DWARF_V3) \
+  X(DW_OP_BREGX, 0x92, 2, DWARF_V3)          X(DW_OP_PIECE, 0x93, 1, DWARF_V3) \
+  X(DW_OP_DEREF_SIZE, 0x94, 1, DWARF_V3)     X(DW_OP_XDEREF_SIZE, 0x95, 1, DWARF_V3) \
+  X(DW_OP_NOP, 0x96, 0, DWARF_V3)            X(DW_OP_PUSH_OBJECT_ADDRESS, 0x97, 0, DWARF_V3) \
+  X(DW_OP_CALL2, 0x98, 1, DWARF_V3)          X(DW_OP_CALL4, 0x99, 1, DWARF_V3) \
+  X(DW_OP_CALL_REF, 0x9a, 1, DWARF_V3)       X(DW_OP_FORM_TLS_ADDRESS, 0x9b, 0, DWARF_V3) \
+  X(DW_OP_CALL_FRAME_CFA, 0x9c, 0, DWARF_V3) X(DW_OP_BIT_PIECE, 0x9d, 2, DWARF_V3) \
+  /* ---------------------------------------------------------------------- */ \
+  X(DW_OP_IMPLICIT_VALUE, 0x9e, 2, DWARF_V4) X(DW_OP_STACK_VALUE, 0x9f, 0, DWARF_V4) \
+  /* ---------------------------------------------------------------------- */ \
+  X(DW_OP_IMPLICIT_POINTER, 0xa0, 2, DWARF_V5) X(DW_OP_ADDRX, 0xa1, 1, DWARF_V5) \
+  X(DW_OP_CONSTX, 0xa2, 1, DWARF_V5)           X(DW_OP_ENTRY_VALUE, 0xa3, 2, DWARF_V5) \
+  X(DW_OP_CONST_TYPE, 0xa4, 3, DWARF_V5)       X(DW_OP_REGVAL_TYPE, 0xa5, 2, DWARF_V5) \
+  X(DW_OP_DEREF_TYPE, 0xa6, 2, DWARF_V5)       X(DW_OP_XDEREF_TYPE, 0xa7, 2, DWARF_V5) \
+  X(DW_OP_CONVERT, 0xa8, 1, DWARF_V5)          X(DW_OP_REINTERPRET, 0xa9, 1, DWARF_V5) \
+/* End of the list */
+
+#define DW_REG_X86_LIST \
+  X(X86, EAX, 0) X(X86, ECX, 1)    X(X86, EDX, 2) X(X86, EBX, 3)  \
+  X(X86, ESP, 4) X(X86, EBP, 5)    X(X86, ESI, 6) X(X86, EDI, 7)  \
+  X(X86, EIP, 8) X(X86, EFLAGS, 9) X(X86, TRAPNO, 10)             \
+  /* ---------------------------------------------------------------------- */ \
+  X(X86, ST0, 11) X(X86, ST1, 12) X(X86, ST2, 13) X(X86, ST3, 14) \
+  X(X86, ST4, 15) X(X86, ST5, 16) X(X86, ST6, 17) X(X86, ST7, 18) \
+  /* ---------------------------------------------------------------------- */ \
+  X(X86, XMM0, 21) X(X86, XMM1, 22) X(X86, XMM2, 23) X(X86, XMM3, 24) \
+  X(X86, XMM4, 25) X(X86, XMM5, 26) X(X86, XMM6, 27) X(X86, XMM7, 28) \
+  /* ---------------------------------------------------------------------- */ \
+  X(X86, MM0, 29) X(X86, MM1, 30) X(X86, MM2, 31) X(X86, MM3, 32) \
+  X(X86, MM4, 33) X(X86, MM5, 34) X(X86, MM6, 35) X(X86, MM7, 36) \
+  /* ---------------------------------------------------------------------- */ \
+  X(X86, FCW, 37) X(X86, FSW, 38) X(X86, MXCSR, 39)                          \
+  X(X86, ES, 40) X(X86, CS, 41) X(X86, SS, 42) X(X86, DS, 43) X(X86, FS, 44) \
+  X(X86, GS, 45) X(X86, TR, 48) X(X86, LDTR, 49)                             \
+  /* End of the list */
+
+#define DW_REG_X64_LIST \
+  X(X64, RAX, 0)  X(X64, RDX, 1)  X(X64, RCX, 2)  X(X64, RBX, 3)  X(X64, RSI, 4)  \
+  X(X64, RDI, 5)  X(X64, RBP, 6)  X(X64, RSP, 7)  X(X64, R8, 8)   X(X64, R9, 9)   \
+  X(X64, R10, 10) X(X64, R11, 11) X(X64, R12, 12) X(X64, R13, 13) X(X64, R14, 14) \
+  X(X64, R15, 15) X(X64, RIP, 16) \
+  /* ---------------------------------------------------------------------- */ \
+  X(X64, XMM0, 17)  X(X64, XMM1, 18)  X(X64, XMM2, 19)  X(X64, XMM3, 20)  \
+  X(X64, XMM4, 21)  X(X64, XMM5, 22)  X(X64, XMM6, 23)  X(X64, XMM7, 24)  \
+  X(X64, XMM8, 25)  X(X64, XMM9, 26)  X(X64, XMM10, 27) X(X64, XMM11, 28) \
+  X(X64, XMM12, 29) X(X64, XMM13, 30) X(X64, XMM14, 31) X(X64, XMM15, 32) \
+  /* ---------------------------------------------------------------------- */ \
+  X(X64, ST0, 33) X(X64, ST1, 34) X(X64, ST2, 35) X(X64, ST3, 36) \
+  X(X64, ST4, 37) X(X64, ST5, 38) X(X64, ST6, 39) X(X64, ST7, 40) \
+  /* ---------------------------------------------------------------------- */ \
+  X(X64, MM0, 41) X(X64, MM1, 42) X(X64, MM2, 43) X(X64, MM3, 44) \
+  X(X64, MM4, 45) X(X64, MM5, 46) X(X64, MM6, 47) X(X64, MM7, 48) \
+  /* ---------------------------------------------------------------------- */ \
+  X(X64, RFLAGS, 49) X(X64, ES, 50) X(X64, CS, 51) X(X64, SS, 52) X(X64, DS, 53)  \
+  X(X64, FS, 54) X(X64, GS, 55) X(X64, FS_BASE, 58) X(X64, GS_BASE, 59)           \
+  X(X64, TR, 62) X(X64, LDTR, 63)                                                 \
+  /* End of the list */
+
+#define DW_REG_ARM_LIST \
+  X(ARM, S0, 64)  X(ARM, S1, 65)  X(ARM, S2, 66)  X(ARM, S3, 67)  X(ARM, S4, 68) X(ARM, S5, 69)   \
+  X(ARM, S6, 70)  X(ARM, S7, 71)  X(ARM, S8, 72)  X(ARM, S9, 73)  X(ARM, S10, 74) X(ARM, S11, 75) \
+  X(ARM, S12, 76) X(ARM, S13, 77) X(ARM, S14, 78) X(ARM, S15, 79) X(ARM, S16, 80) X(ARM, S17, 81) \
+  X(ARM, S18, 82) X(ARM, S19, 83) X(ARM, S20, 84) X(ARM, S21, 85) X(ARM, S22, 86) X(ARM, S23, 87) \
+  X(ARM, S24, 88) X(ARM, S25, 89) X(ARM, S26, 90) X(ARM, S27, 91) X(ARM, S28, 92) X(ARM, S29, 93) \
+  X(ARM, S30, 94) X(ARM, S31, 95) \
+  /* ---------------------------------------------------------------------- */ \
+  X(ARM, F0, 96)  X(ARM, F1, 97)  X(ARM, F2, 98)  X(ARM, F3, 99) \
+  X(ARM, F4, 100) X(ARM, F5, 101) X(ARM, F6, 102) X(ARM, F7, 103) \
+  /* ---------------------------------------------------------------------- */ \
+  X(ARM, WCGR0, 104) X(ARM, WCGR1, 105) X(ARM, WCGR2, 106) \
+  X(ARM, WCGR3, 107) X(ARM, WCGR4, 108) X(ARM, WCGR5, 109) \
+  X(ARM, WCGR6, 110) X(ARM, WCGR7, 111) \
+  /* ---------------------------------------------------------------------- */ \
+  X(ARM, WR0, 112)  X(ARM, WR1, 113)  X(ARM, WR2, 114)  X(ARM, WR3, 115)  X(ARM, WR4, 116)  X(ARM, WR5, 117) \
+  X(ARM, WR6, 118)  X(ARM, WR7, 119)  X(ARM, WR8, 120)  X(ARM, WR9, 121)  X(ARM, WR10, 122) X(ARM, WR11, 123) \
+  X(ARM, WR12, 124) X(ARM, WR13, 125) X(ARM, WR14, 126) X(ARM, WR15, 127) \
+  /* ---------------------------------------------------------------------- */ \
+  X(ARM, SPSR, 128) X(ARM, SPSR_FIQ, 129) X(ARM, SPSR_IRQ, 130) X(ARM, SPSR_ABT, 131) \
+  X(ARM, SPSR_UND, 132) X(ARM, SPSR_SVC, 133)                                         \
+  /* ---------------------------------------------------------------------- */ \
+  X(ARM, R8_USR,   144) X(ARM, R9_USR,   145) X(ARM, R10_USR,  146) X(ARM, R11_USR,  147) \
+  X(ARM, R12_USR,  148) X(ARM, R13_USR,  149) X(ARM, R14_USR,  150) X(ARM, R8_FIQ,   151) \
+  X(ARM, R9_FIQ,   152) X(ARM, R10_FIQ,  153) X(ARM, R11_FIQ,  154) X(ARM, R12_FIQ,  155) \
+  X(ARM, R13_FIQ,  156) X(ARM, R14_FIQ,  157) X(ARM, R13_RIQ,  158) X(ARM, R14_RIQ,  159) \
+  X(ARM, R14_ABT,  160) X(ARM, R13_ABT,  161) X(ARM, R14_UND,  162) X(ARM, R13_UND,  163) \
+  X(ARM, R14_SVC,  164) X(ARM, R13_SVC,  165) \
+  /* ---------------------------------------------------------------------- */ \
+  X(ARM, WC0, 192) X(ARM, WC1, 193) X(ARM, WC2, 194) X(ARM, WC3, 195) \
+  X(ARM, WC4, 196) X(ARM, WC5, 197) X(ARM, WC6, 198) X(ARM, WC7, 199) \
+  /* ---------------------------------------------------------------------- */ \
+  X(ARM, D0, 256)  X(ARM, D1, 257)  X(ARM, D2, 258)  X(ARM, D3, 259)  X(ARM, D4, 260)  \
+  X(ARM, D5, 261)  X(ARM, D6, 262)  X(ARM, D7, 263)  X(ARM, D8, 264)  X(ARM, D9, 265)  \
+  X(ARM, D10, 266) X(ARM, D11, 267) X(ARM, D12, 268) X(ARM, D13, 269) X(ARM, D14, 270) \
+  X(ARM, D15, 271) X(ARM, D16, 272) X(ARM, D17, 273) X(ARM, D18, 274) X(ARM, D19, 275) \
+  X(ARM, D20, 276) X(ARM, D21, 277) X(ARM, D22, 278) X(ARM, D23, 279) X(ARM, D24, 280) \
+  X(ARM, D25, 281) X(ARM, D26, 282) X(ARM, D27, 283) X(ARM, D28, 284) X(ARM, D29, 285) \
+  X(ARM, D30, 286) X(ARM, D31, 287) \
+  /* End of the list */
+
+#define DW_ARCH_LIST \
+  I(X86) I(X64) I(ARM)
+
+#define X(arch, name, value)    DW_REG_##arch##_##name = value,
+#define I(arch) typedef enum { DW_REG_##arch##_NULL, DW_REG_##arch##_LIST DW_REG_##arch##_MAX } DwRegType##arch;
+  DW_ARCH_LIST
+#undef I
+#undef X
+
+typedef enum 
+{
+  DW_TAG_NULL,
+#define X(name, value, version) name = value,
+  DW_TAG_LIST
+#undef X
+  DW_TAG_LO_USER                  = 0x4080,
+  DW_TAG_HI_USER                  = 0xffff
+} DwTagType;
+
+typedef enum 
+{
+#define X(name, value, ver, class_type) name = value,
+  DW_ATTRIB_LIST
+#undef X
+  DW_AT_LO_USER = 0x2000,
+  DW_AT_HI_USER = 0x3fff
+} DwAttribType;
+
+typedef enum 
+{
+#define X(name, value, ver, class_type) name = value,
+  DW_FORM_LIST
+#undef X
+
+  DW_FORM_INVALID
+} DwForm;
+
+typedef enum 
+{
+  DW_ATE_NULL,
+
+#define X(name, value, version) name = value,
+  DW_ATE_LIST
+#undef X
+
+  DW_ATE_LO_USER = 0x80,
+  DW_ATE_HI_USER = 0xff
+} DwAttribTypeEncoding;
+
+typedef enum 
+{
+  DW_OP_NULL,
+
+#define X(name, value, num_ops, version) name = value,
+  DW_OP_LIST
+#undef X
+
+  DW_OP_LO_USER = 0xe0,
+  DW_OP_HI_USER = 0xff
+} DwOperation;
+
+typedef enum 
+{
+  DW_LANG_INVALID       = 0x00, // TODO(nick): rename invalid to null
+  DW_LANG_C89           = 0x01, 
+  DW_LANG_C             = 0x02, 
+  DW_LANG_ADA83         = 0x03,
+  DW_LANG_C_PLUS_PLUS   = 0x04, 
+  DW_LANG_COBOL74       = 0x05, 
+  DW_LANG_COBOL85       = 0x06,
+  DW_LANG_FORTAN77      = 0x07, 
+  DW_LANG_FORTAN90      = 0x08, 
+  DW_LANG_PASCAL83      = 0x09,
+  DW_LANG_MODULA2       = 0x0A, 
+  DW_LANG_JAVA          = 0x0B, 
+  DW_LANG_C99           = 0x0C,
+  DW_LANG_ADA95         = 0x0D, 
+  DW_LANG_FORTAN95      = 0x0E, 
+  DW_LANG_PLI           = 0x0F,
+  DW_LANG_OBJ_C         = 0x10, 
+  DW_LANG_OBJ_CPP       = 0x11, 
+  DW_LANG_UPC           = 0x12,
+  DW_LANG_D             = 0x13, 
+  DW_LANG_PYTHON        = 0x14,
+  DW_LANG_LO_USER = 0x8000, 
+  DW_LANG_HI_USER = 0xffff
+} DwLang;
+
+typedef enum
+{
+  DW_LNS_SPECIAL_OPCODE     = 0x00, 
+  DW_LNS_COPY               = 0x01,
+  DW_LNS_ADVANCE_PC         = 0x02, 
+  DW_LNS_ADVANCE_LINE       = 0x03,
+  DW_LNS_SET_FILE           = 0x04, 
+  DW_LNS_SET_COLUMN         = 0x05,
+  DW_LNS_NEGATE_STMT        = 0x06, 
+  DW_LNS_SET_BASIC_BLOCK    = 0x07,
+  DW_LNS_CONST_ADD_PC       = 0x08, 
+  DW_LNS_FIXED_ADVANCE_PC   = 0x09,
+  DW_LNS_SET_PROLOGUE_END   = 0x0A, 
+  DW_LNS_SET_EPILOGUE_BEGIN = 0x0B,
+  DW_LNS_SET_ISA            = 0x0C
+} DwStdOpcode;
+
+typedef enum 
+{
+  DW_LNE_UNDEFINED         = 0x00, 
+  DW_LNE_END_SEQUENCE      = 0x01, 
+  DW_LNE_SET_ADDRESS       = 0x02, 
+  DW_LNE_DEFINE_FILE       = 0x03, 
+  DW_LNE_SET_DISCRIMINATOR = 0x04,
+  DW_LNE_LO_USER = 0x80, DW_LNE_HI_USER = 0xff
+} DwExtOpcode;
+
+typedef enum 
+{
+  DW_UT_RESERVED      = 0x00,
+  DW_UT_COMPILE       = 0x01, 
+  DW_UT_TYPE          = 0x02,
+  DW_UT_PARTIAL       = 0x03, 
+  DW_UT_SKELETON      = 0x04,
+  DW_UT_SPLIT_COMPILE = 0x05, 
+  DW_UT_SPLIT_TYPE    = 0x06,
+  DW_UT_LO_USER       = 0x80, 
+  DW_UT_HI_USER       = 0xff
+} DwUnitType;
+
+typedef enum
+{
+  DW_IDENTIFIER_CASE_SENSITIVE,
+  DW_IDENTIFIER_UP_CASE,
+  DW_IDENTIFIER_DOWN_CASE,
+  DW_IDENTIFIER_CASE_INSENSITIVE
+} DwIdentifierCaseType;
+
+// -------------------------------------------------------------------------------- 
+
+#define DW_CAST(x, u) (x)(u)
+#define DW_PTR_DIFF_BYTES(a, b) DW_CAST(dw_uint, DW_CAST(U8 *, SYMS_MAX(a, b)) - DW_CAST(U8 *, SYMS_MIN(a, b)))
+#define DW_INVALID_OFFSET SYMS_UINT64_MAX
+#define DW_INVALID_VALUE  (~0u)
+
+/* NOTE(nick): Size of this type is runtime specific and depends on DWARF mode.
+ * 32bit - 4 bytes, 
+ * 64bit - 8 bytes. 
+ */
+typedef void DwMSize;
+
+/* NOTE(nick): An entry that points to a segment and address pair. 
+ * Segmenet part is optional. */
+typedef void DwSegAddr;
+
+typedef enum 
+{
+  DW_MODE_NULL,
+  DW_MODE_32BIT,
+  DW_MODE_64BIT
+} DwMode;
+
+typedef enum 
+{
+  DWARF_INVALID_VERSION,
+  DWARF_V1,
+  DWARF_V2,
+  DWARF_V3,
+  DWARF_V4,
+  DWARF_V5,
+  DWARF_LAST_VERSION = DWARF_V5
+} DwVersion;
+
+typedef enum 
+{
+  DW_SEC_NULL,
+  DW_SEC_ABBREV,
+  DW_SEC_ARANGES,
+  DW_SEC_FRAME,
+  DW_SEC_INFO,
+  DW_SEC_LINE,
+  DW_SEC_LOC,
+  DW_SEC_MACINFO,
+  DW_SEC_PUBNAMES,
+  DW_SEC_PUBTYPES,
+  DW_SEC_RANGES,
+  DW_SEC_STR,
+  DW_SEC_ADDR,
+  DW_SEC_LOCLISTS,
+  DW_SEC_RNGLISTS,
+  DW_SEC_STR_OFFSETS,
+  DW_SEC_LINE_STR,
+  DW_SEC_MAX
+} DwSecType;
+
+typedef struct DwImgSec 
+{
+  size_t data_len;
+  void *data;
+} DwImgSec;
+
+#if 1
+typedef struct DwBinRead // TODO(nick): syms_core.h has an alternative for this
+{
+  syms_bool err;
+
+  U8 addr_size;
+  DwMode mode;
+
+  SymsAddr off;
+  SymsAddr max;
+  void *data;
+} DwBinRead;
+#endif
+
+typedef struct DwBinWrite 
+{
+  SymsAddr off;
+  SymsAddr max;
+  void *data;
+} DwBinWrite;
+
+typedef enum 
+{
+  DW_NAME_TABLE_PUBTYPES,
+  DW_NAME_TABLE_PUBNAMES,
+  DW_NAME_TABLE_MAX
+} DwNameTableIndex;
+
+typedef struct DwNameTableKeyValue 
+{
+  void *entry;
+  struct DwNameTableKeyValue *next;
+} DwNameTableKeyValue;
+
+typedef struct DwNameTable 
+{
+  struct DwNameTableKeyValue *keys[4096];
+} DwNameTable;
+
+#define DW_REGREAD_SIG(name) dw_uint name(void *context, SymsArch arch, dw_uint reg_index, void *read_buffer, dw_uint read_buffer_max)
+typedef DW_REGREAD_SIG(dw_regread_sig);
+/* NOTE(nick): This callback is used during decoding of the addresses(e.g. locals, globals, lines).
+ * If user chose to omit this callback (i.e. passes a NULL) library won't be able to resolve addresses. */
+
+#define DW_REGWRITE_SIG(name) dw_uint name(void *context, SymsArch arch, dw_uint reg_index, void *value, dw_uint value_size)
+typedef DW_REGWRITE_SIG(dw_regwrite_sig);
+
+#define DW_MEMREAD_SIG(name) syms_bool name(void *context, SymsAddr va, void *read_buffer, dw_uint num_read)
+typedef DW_MEMREAD_SIG(dw_memread_sig);
+/* NOTE(nick): This callback is primarily used for decoding variables' addresses. 
+ * If user chose to omit this callback (i.e. passes a NULL) library won't be able to resolve addresses. */
+
+#define DW_NEXT_SEC_SIG(name) syms_bool name(void *user_context, dw_uint iter_index, DwImgSec *out_sec)
+typedef DW_NEXT_SEC_SIG(dw_next_sec_sig);
+/* NOTE(nick): This callback allows this lib to resolve external references to other .deubg_info sections. 
+ * If user chose to omit this callback (i.e passes a NULL) external references will be ignroed. */
+
+#define DW_HAS_CU_LOADED(x) ((x)->cu_num && (x)->cu_arr)
+#define DW_HAS_SECTION(dw, sec) ((dw)->secs[sec].data_len > 0)
+typedef struct DwContext 
+{
+  /* NOTE(nick): If there is a section that is larger 
+   * than a 32bit value can hold, mode is set to 64bit. */
+  DwMode mode;
+  U8 msize_byte_count;
+
+  /* NOTE(nick): Target's architecture, this value has to be read from
+   * the executable's header and passed to this library. */
+  SymsArch arch;
+
+  /* NOTE(nick): User's contexts for the callbacks. */
+  void *next_info_ctx;
+
+  /* NOTE(nick): Optional callbacks, if one is not present the operation
+   * that depends on it is ignored and the result value is set to an invalid one.
+   * For more info you can read macro comments above. */
+  dw_next_sec_sig *next_sec;
+
+  /* NOTE(nick): Sections that user provides during the init step. */
+  DwImgSec secs[DW_SEC_MAX];
+
+  /* NOTE(nick): Number of ids that ".debug_abbrev" contains. */
+  U32 abbrev_num;
+
+  /* NOTE(nick): Number of Compile Units that ".debug_info" contains. */
+  U32 cu_num;
+
+  /* NOTE(nick): Number of strings that ".debug_pubnames" contains. */
+  U32 pubnames_str_num;
+
+  /* NOTE(nick): Number of strings that ".debug_pubtypes" contains. */
+  U32 pubtypes_str_num;
+
+  /* NOTE(nick): Internally-linked hash tables; used for fast name to offset maps. */
+  struct DwNameTable *name_tables[DW_NAME_TABLE_MAX];
+} DwContext;
+
+typedef struct DwTag 
+{
+  void *cu;
+  void *info;
+  void *abbrev;
+} DwTag;
+
+typedef struct DwSegOffArray 
+{
+  dw_uint segoff_size;
+  dw_uint segsel_size;
+  dw_uint num;
+  void *entries;
+} DwSegOffArray;
+
+#define dw_off_array_has_entries(x) ((x).num && (x).entries)
+typedef struct DwOffArray 
+{
+  U8 entry_len;
+  dw_uint num;
+  void *entries;
+} DwOffArray;
+
+typedef struct DwAbbrevTableEntry 
+{
+  dw_uint id;
+  SymsOffset off;
+} DwAbbrevTableEntry;
+
+typedef struct DwAbbrevTable
+{
+  U32 entry_count;
+  DwAbbrevTableEntry *entries;
+} DwAbbrevTable;
+
+typedef struct DwCompileUnit 
+{
+  /* NOTE(nick): For internal use. Don't touch this! */
+  struct DwContext *dwarf;
+
+  DwMode mode;
+
+  /* NOTE(nick): Series of segment/address pairs read from ".debug_addr" section. */
+  DwSegOffArray addrs_arr;
+
+  /* NOTE(nick): Series of segment/address pairs read from the ".debug_loclists" section. */
+  DwSegOffArray loclists_arr;
+
+  /* NOTE(nick): Series of segment/address pairs read from the ".debug_rnglists" section. */
+  DwSegOffArray rnglists_arr;
+
+  /* NOTE(nick): Offsets into the ".debug_str" section read from the ".debug_str_offsets" section. */
+  DwOffArray stroffs_arr;
+
+  DwAbbrevTable abbrev_table;
+
+  U32 index;
+
+  /* NOTE(nick): Size of an address. */
+  U8 addr_size;
+
+  /* NOTE(nick): Version, boring but someone might in someway rely on this info, I guess. */
+  U16 ver;
+
+  /* NOTE(nick): Size of the instructions that contributed to a code segment. */
+  U32 len;
+
+  /* NOTE(nick): This was introduced in DWARF5 to indicate type of the
+   * compilation unit. There is no need for me to repeat spec here, so
+   * you can find description of this in chapter 3.2.1. */
+  DwUnitType unit_type;
+
+  /* NOTE(nick): Programming language that was used. */
+  DwLang lang;
+
+  /* NOTE(nick): Identifier comparison option. */
+  DwIdentifierCaseType case_type;
+
+  /* NOTE(nick): Indicates to decoder that UTF-8 decoder needs to be used
+   * to read strings. */
+  syms_bool use_utf8;
+
+  /* NOTE(nick): Module-base-relative address of a first instruction 
+   * to which this compile unit contributed. 
+   *
+   * TIP: You can use this to iterate over the segments to find 
+   * out which one is refed by this compile unit. */
+  U64 rva;
+
+  /* NOTE(nick): Pointer to the first byte of the Compile Unit header
+   * in the ".debug_info" section. */
+  void *info_data_start;
+
+  /* NOTE(nick): An absolute offset into the ".debug_info" section,
+   * where info for this compile unit starts. */
+  SymsAddr info_base;
+
+  /* NOTE(nick): Indicates how many bytes after "info_base" belong to this compile unit. */
+  SymsAddr info_len;
+
+  SymsAddr range_off;
+
+  /* NOTE(nick): Offset where the attributes start. */
+  SymsAddr attribs_off;
+
+  /* NOTE(nick): An absolute offset into the ".debug_abbrev" section,
+   * where abbrev info for this compile unit starts. */
+  SymsAddr abbrev_base;
+
+  /* NOTE(nick): An absolute offset into the ".debug_line" section,
+   * where line info for this compile unit starts. */
+  SymsAddr line_base;
+
+  /* NOTE(nick): Name of the compiler that produced this DWARF info. */
+  SymsString producer;
+
+  /* NOTE(nick): Compile directory. */
+  SymsString compile_dir;
+
+  /* NOTE(nick): Name of the file that was used to produce this compile unit. */
+  SymsString name;
+
+  SymsString dwo_name;
+
+  u64 dwo_id;
+} DwCompileUnit;
+
+typedef struct DwCuIter 
+{
+  syms_bool err;
+
+  struct DwContext *dwarf;
+
+  /* NOTE(nick): Number of Compile Units that were read. */
+  U32 num_read;
+
+  /* NOTE(nick): Offset into the .debug_info section where next
+   * CompileUnit starts. */
+  SymsAddr next_cu;
+} DwCuIter;
+
+/* NOTE(nick): A layout of .debug_abbrev
+ *
+ * DW_TAG_COMPILE_UNIT [has children]
+ *   DW_AT_PRODUCER 
+ *   DW_AT_LANGUAGE
+ *   DW_AT_NAME
+ *   DW_AT_NULL 
+ *     DW_TAG_SUBPROGRAM [no children]
+ *       DW_AT_NAME
+ *       DW_AT_LOW_PC
+ *       DW_AT_HIGH_PC
+ *       DW_AT_NULL
+ *     DW_TAG_POINTER [no children]
+ *       DW_AT_IMPORT
+ *       DW_AT_VALUE
+ *       DW_AT_NULL
+ * DIE_END
+ *
+ */
+enum 
+{
+  DW_ABBREV_ENTRY_TYPE_NULL,
+  DW_ABBREV_ENTRY_TYPE_TAG_INFO,
+  DW_ABBREV_ENTRY_TYPE_ATTRIB_INFO,
+  /* NOTE(nick): Indicates end of a sequence of attributes,
+   * next entry is either a tag info or a DIE end. */
+  DW_ABBREV_ENTRY_TYPE_ATTRIB_INFO_NULL,
+  DW_ABBREV_ENTRY_TYPE_DIE_BEGIN,
+  DW_ABBREV_ENTRY_TYPE_DIE_END
+};
+typedef U32 DwAbbrevEntryType;
+
+typedef struct DwAbbrevAttribInfo 
+{
+  /* NOTE(nick): Name of attribute */
+  U64 name;
+  /* NOTE(nick): Format of the value that is stored in the .debug_info */
+  U64 form;
+  U64 implicit_const;
+  syms_bool has_implicit_const;
+} DwAbbrevAttribInfo;
+
+typedef struct DwAbbrevTagInfo 
+{
+  U64 id;
+  U64 tag;
+  /* NOTE(nick): If syms_true all subsequent entries are children of this tag,
+   * until you receive DW_ABBREV_ENTRY_ATTRIB_INFO_NULL. */
+  U8 has_children;
+} DwAbbrevTagInfo;
+
+typedef struct DwAbbrevEntry 
+{
+  /* NOTE(nick): Offsets into .debug_abbrev section that indicate
+   * where bytes for this entry start and end, respectively */
+  SymsAddr data_lo;
+  SymsAddr data_hi;
+
+  /* NOTE(nick): Use this field to access field from union below. */
+  DwAbbrevEntryType type;
+  union {
+    DwAbbrevTagInfo tag_info;
+    DwAbbrevAttribInfo attrib_info;
+  } u;
+} DwAbbrevEntry;
+
+enum 
+{
+  DW_ABBREV_ITER_STATE_NULL,
+  DW_ABBREV_ITER_STATE_EXPECT_TAG_INFO,
+  DW_ABBREV_ITER_STATE_EXPECT_ATTRIB_INFO,
+  DW_ABBREV_ITER_STATE_EMIT_DIE_BEGIN
+};
+typedef U32 DwAbbrevIterState;
+
+typedef struct DwAbbrevIter 
+{
+  DwAbbrevIterState state;
+  DwBinRead data;
+} DwAbbrevIter;
+
+typedef enum 
+{
+  DW_ATTRIB_VALUE_NULL,
+  DW_ATTRIB_VALUE_STRING,
+  DW_ATTRIB_VALUE_CONST,
+  DW_ATTRIB_VALUE_SECOFF
+} DwAttribValueType;
+
+typedef enum 
+{
+  DW_AT_CLASS_INVALID    = 0,
+  DW_AT_CLASS_ADDRESS    = (1 << 0),
+  DW_AT_CLASS_BLOCK      = (1 << 1),
+  DW_AT_CLASS_CONST      = (1 << 2),
+  DW_AT_CLASS_EXPRLOC    = (1 << 3),
+  DW_AT_CLASS_FLAG       = (1 << 4),
+  DW_AT_CLASS_LINEPTR    = (1 << 5),
+  DW_AT_CLASS_LOCLIST    = (1 << 6),
+  DW_AT_CLASS_LOCLISTPTR = (1 << 7),
+  DW_AT_CLASS_MACPTR     = (1 << 8),
+  DW_AT_CLASS_RNGLISTPTR = (1 << 9),
+  DW_AT_CLASS_RNGLIST    = (1 << 10),
+  DW_AT_CLASS_REFERENCE  = (1 << 11),
+  DW_AT_CLASS_STRING     = (1 << 12),
+
+  /* NOTE(nick): Introduced in version 5 */
+  DW_AT_CLASS_STROFFSETSPTR = (1 << 13),
+  DW_AT_CLASS_ADDRPTR       = (1 << 14),
+
+  DW_AT_CLASS_UNDEFINED = (1 << 15)
+} DwAttribClass;
+
+typedef struct DwAttribExprloc 
+{
+  U64 len;
+  U8 *data;
+} DwAttribExprloc;
+
+typedef struct DwFormBlock
+{
+  U64 len;
+  void *data;
+} DwFormBlock;
+
+typedef struct DwRef
+{
+  SymsOffset info;
+} DwRef;
+
+typedef struct DwAttrib 
+{
+  SymsOffset info_lo;
+  SymsOffset info_hi;
+
+  DwTag tag;
+  DwAttribType name;
+  DwForm form;
+  DwAttribClass value_class;
+  union {
+    U8 flag;
+    U64 data;
+    S64 sdata;
+    U64 udata;
+    U64 ref;
+    U64 addr;
+    U64 cnst;
+    U64 addrx;
+    U64 loclistx;
+    U64 strp;
+    U64 strx;
+    U64 rnglistx;
+    SymsAddr sec_offset;
+    SymsString string;
+    DwFormBlock block;
+    struct {
+      U64 lo;
+      U64 hi;
+    } data16;
+
+    struct {
+      U64 len;
+      U8 *data;
+    } exprloc;
+  } form_value;
+
+  union {
+    SymsAddr address;
+    U64 flag;
+    SymsAddr rnglistptr;
+    SymsAddr rnglist;
+    SymsAddr addrptr;
+    SymsAddr stroffptr;
+    SymsAddr lineptr;
+    SymsAddr loclistptr;
+    SymsAddr loclist;
+    SymsAddr macptr;
+    SymsString  string;
+    DwRef ref;
+    DwAttribExprloc exprloc;
+
+    DwFormBlock block;
+
+    struct {
+      SymsAddr lo;
+      SymsAddr hi;
+    } cnst16;
+  } value;
+} DwAttrib;
+
+typedef struct DwAttribIter 
+{
+  /* NOTE(nick): Compile unit that owns attributes. */
+  struct DwCompileUnit *cu;
+
+  struct DwAbbrevIter abbrev_iter;
+
+  /* NOTE(nick): Subsection of the .debug_info memory. */
+  DwBinRead info;
+
+  /* NOTE(nick): When no attributes left to read this is set to syms_true.*/
+  syms_bool is_exhausted;
+
+  DwTagType tag_type;
+
+  /* NOTE(nick): Unique ID that belongs to the entry that lives in
+   * the .debug_info section. */
+  U64 abbrev_id;
+  
+  /* NOTE(nick): Relative to a compile unit offset into the .debug_info section,
+   * where attribute info starts. */
+  SymsAddr info_off;
+
+  /* NOTE(nick): Offset into the .debug_abbrev section, where
+   * starts abbreviation info for the "info_off". */
+  SymsAddr abbrev_off;
+
+  /* NOTE(nick): Length in bytes of the tag header. */
+  SymsAddr abbrev_header_len;
+
+  /* NOTE(nick): Set syms_true when a chain of tags follow this one.
+   * Chain is ended by a 0 in the .debug_info section.
+   *
+   * You can detect end of the chain by using the "dw_attrib_iter_init"
+   * and if the "is_exhausted" is set to syms_true right after init then
+   * this is end. */
+  syms_bool has_children;
+
+  /* NOTE(nick): This member is for resolving DW_FORM_REF_ADDR.
+   * This form allows producer of the DWARF info to make
+   * references to external DWARF sections. DW_FORM_REF_ADDR
+   * supplies the offset into a section, and this is a symbolic
+   * description of that DWARF section; it is set when one is encountered.
+   * If the symbolic description was found this field is used to store it,
+   * otherwise it is null. When this field is null, and iterator encounters
+   * DW_FORM_REF_ADDR we resolve this form as an offset into the .debug_info
+   * that was supplied by the user during the init step. */
+  char *ref_addr_desc;
+} DwAttribIter;
+
+typedef struct DwTagIter
+{
+  S32 depth;
+  DwCompileUnit *cu;
+  SymsOffset info_off;
+} DwTagIter;
+
+typedef dw_uint DwLn;
+typedef dw_uint DwCol;
+typedef dw_uint DwFileIndex;
+typedef dw_uint DwDirIndex;
+
+typedef struct DwLineIterFile 
+{
+  SymsString path;
+  DwDirIndex dir_index;
+  U64 modify_time;
+  U64 file_size;
+  struct DwLineIterFile *next;
+} DwLineIterFile;
+
+typedef struct DwLineIterDir 
+{
+  SymsString path;
+  struct DwLineIterDir *next;
+} DwLineIterDir;
+
+typedef struct DwLine 
+{
+  /* NOTE(nick): Address of a machine instruction. */
+  SymsAddr address;
+
+  /* NOTE(nick): This is used by the VLIW instructions to
+   * indicate index of operation inside the instruction.  */
+  U32 op_index;
+
+  /* NOTE(nick): Line table doesn't contain full path to a file,
+   * instead DWARF encodes path as two indices. 
+   * First index will point into a directory table, 
+   * and second points into a file name table. */
+  DwFileIndex file_index;
+
+  /* NOTE(nick): Source-level line number. */
+  DwLn line;
+
+  /* NOTE(nick): Source-level column. */
+  DwCol column;
+
+  /* NOTE(nick): Indicates that "va" points to place suitable for a breakpoint. */
+  syms_bool is_stmt;
+
+  /* NOTE(nick): Indicates that the "va" is inside a basic block. */
+  syms_bool basic_block;
+
+  /* NOTE(nick): Indicates that "va" points to place where function starts.
+   * Usually prologue is the place where compiler emits instructions to 
+   * prepare stack for a function.  */
+  syms_bool prologue_end;
+
+  /* NOTE(nick): Indicates that "va" points to section where function exits and unwinds stack. */
+  syms_bool epilogue_begin;
+
+  /* NOTE(nick): Instruction set that is used. */
+  U64 isa;
+
+  /* NOTE(nick): Arbitrary id that indicates to which block 
+   * these instructions belong. */
+  U64 discriminator;
+
+  /* NOTE(nick): Indicates that "va" points to the first instruction
+   * in the instruction block that follows. */
+  syms_bool end_sequence;
+} DwLine;
+
+typedef struct DwLineFile 
+{
+  SymsString file_name;
+  DwFileIndex file_index;
+  DwDirIndex dir_index;
+  U64 modify_time;
+  U64 file_size;
+} DwLineFile;
+
+typedef struct DwLineDir 
+{
+  DwDirIndex dir_index;
+  SymsString dir_name;
+} DwLineDir;
+
+typedef enum 
+{
+  DW_LINE_ITER_OP_NULL,
+  DW_LINE_ITER_OP_LINE,
+  DW_LINE_ITER_OP_DEFINE_FILE
+} DwLineIterOpType;
+
+typedef struct DwLineIterOp 
+{
+  DwLineIterOpType type;
+  union {
+    DwLine line;
+    DwLineFile file;
+  } u;
+} DwLineIterOp;
+
+typedef struct DwLineIter 
+{
+  /* NOTE(nick): portion of .debug_line that is used to read the opcodes for the
+   * state machine. */
+  DwBinRead linesec;
+
+  /* NOTE(nick): Machine state that gets updated by reading opcodes. */
+  DwLine state;
+
+  /* NOTE(nick): All addresses in the line table are based on this address. */
+  SymsAddr base_addr;
+
+  /* NOTE(nick): Directory path that was used during at the compile time. */
+  SymsString compile_dir;
+
+  SymsString compile_file;
+
+  /* NOTE(nick): Offset into the "linesec" where opcodes start. */
+  SymsAddr prog_off;
+
+  DwBinRead dirs;
+  DwBinRead files;
+  dw_uint dir_count;
+  dw_uint file_count;
+
+  /* NOTE(nick): Data that is read from the section header, and
+   * is used by state machine when reading opcodes. */
+  U64 unit_length;
+  U16 ver;
+  U64 header_len;
+  U8 min_inst_len;
+  U8 max_ops_for_inst;
+  U8 default_is_stmt;
+  S8 line_base;
+  U8 line_range;
+  U8 opcode_base;
+  U64 std_opcode_len;
+
+  /* NOTE(nick): This is an array that contains number of operands per opcode. 
+   * This array is used to skip opcodes that aren't part of the spec. 
+   * "opcode_lens" is a pointer into the "linesec" memory. */
+  U32        num_opcode_lens;
+  U8 *opcode_lens;
+} DwLineIter;
+
+typedef struct DwLineMap 
+{
+  /* NOTE(nick): Address of the first instruction that belongs to this line. */
+  SymsAddr va;
+
+  /* NOTE(nick): Size in bytes of all instructions that contribute to the line. */
+  U32 size;
+
+  /* NOTE(nick): Line number */
+  U64 ln;
+
+  /* NOTE(nick): Coulmn number. Use this carefully, not all debug formats export
+   * column information. */
+  U32 col;
+
+  SymsString dir;
+  SymsString filename;
+} DwLineMap;
+
+typedef struct DwNameIter 
+{
+  DwContext *dwarf;
+  DwBinRead sec;
+  SymsAddr unit_start;
+  SymsAddr unit_end;
+  U16 unit_ver;
+  void *current_entry;
+  SymsAddr cu_info_off;
+  U64 cu_info_len;
+} DwNameIter;
+
+typedef struct DwNameTableEntry32 
+{
+  U32 info_off;
+  char name[1];
+} DwNameTableEntry32;
+
+typedef struct DwNameTableEntry64 
+{
+  U64 info_off;
+  char name[1];
+} DwNameTableEntry64;
+
+typedef enum 
+{
+  DW_ENCODED_LOCATION_NULL,
+  DW_ENCODED_LOCATION_EXPR,
+  DW_ENCODED_LOCATION_RVA
+} DwEncodedLocationType;
+
+typedef struct DwEncodedLocationExpr 
+{
+  void *ops;
+  dw_uint ops_size;
+  SymsAddr frame_base;
+  SymsAddr member_location;
+  SymsAddr cfa;
+} DwEncodedLocationExpr;
+
+typedef struct DwEncodedLocationRva 
+{
+  struct DwCompileUnit *cu;
+  U64 off;
+} DwEncodedLocationRva;
+
+typedef struct DwEncodedLocation 
+{
+  DwCompileUnit *cu;
+  DwEncodedLocationType type;
+  union {
+    DwEncodedLocationExpr expr;
+    DwEncodedLocationRva  rva;
+  } u;
+} DwEncodedLocation;
+
+enum 
+{
+  DW_SIZE_TYPE_NULL,
+  DW_SIZE_TYPE_BIT,
+  DW_SIZE_TYPE_BYTE
+};
+typedef U8 DwSizeType;
+
+typedef struct DwCommonAttribs 
+{
+  U32 decl_file;
+  U32 decl_ln;
+  U32 call_file;
+  U32 call_ln;
+  U32 len;
+  SymsAddr rva;
+  SymsAddr range_off;
+  DwRef type_tag;
+  DwRef sibling_tag;
+  DwRef specification;
+  SymsString linkage_name;
+  SymsString name;
+} DwCommonAttribs;
+
+typedef struct DwProc 
+{
+  /* NOTE(nick): Byte length of all instructions that this procedure encapsulated. */
+  U32 len; 
+
+  U32 decl_ln;
+  U32 decl_file;
+
+  U32 call_ln;
+  U32 call_file;
+
+  SymsOffset range_off;
+
+  /* NOTE(nick): Name that was given by the user to this procedure. */
+  SymsString name;
+
+  DwTag type_tag;
+
+  /* NOTE(nick): Address of a first instruction. */
+  DwEncodedLocation encoded_va;
+
+  DwAttrib frame_base;
+} DwProc;
+
+typedef struct DwScope 
+{
+  /* NOTE(nick): OPTIONAL! Name of the scope??? Don't know what is this, but it is there. */
+  SymsString name;
+
+  /* NOTE(nick): OPTIONAL! Line number where scope was declared. */
+  U32 decl_ln;
+
+  /* NOTE(nick): OPTIONAL! File where scope was declared. */
+  U32 decl_file;
+
+  /* NOTE(nick): Address where first instruction that belongs to this scope starts. */
+  U64 rva;
+
+  /* NOTE(nick): Byte length of all instructions that this scope encapsulated. */
+  U32 len;
+
+  /* NOTE(nick): OPTIONAL! compiler might put it in so we can skip scope bytes quickly. */
+  DwTag sibling_tag;
+} DwScope;
+
+typedef struct DwClass 
+{
+  DwSizeType size_type;
+  U32 len;
+  SymsString name;
+} DwClass;
+
+typedef enum 
+{
+  DW_TYPE_NULL,
+
+  DW_TYPE_FLOAT16,
+  DW_TYPE_FLOAT32,
+  DW_TYPE_FLOAT64,
+  DW_TYPE_FLOAT48,
+  DW_TYPE_FLOAT80,
+  DW_TYPE_FLOAT128,
+
+  DW_TYPE_INT8,
+  DW_TYPE_INT16,
+  DW_TYPE_INT32,
+  DW_TYPE_INT64,
+  DW_TYPE_INT128,
+
+  DW_TYPE_UINT8,
+  DW_TYPE_UINT16,
+  DW_TYPE_UINT32,
+  DW_TYPE_UINT64,
+  DW_TYPE_UINT128,
+
+  DW_TYPE_STRUCT,
+  DW_TYPE_UNION,
+  DW_TYPE_CLASS,
+
+  DW_TYPE_ENUM,
+  DW_TYPE_PROC,
+  DW_TYPE_PROC_PARAM,
+  DW_TYPE_VOID,
+  DW_TYPE_BOOL,
+
+  DW_TYPE_PTR,
+  DW_TYPE_ARR,
+
+  DW_TYPE_TYPEDEF
+} DwTypeKind;
+
+typedef enum 
+{
+  /* NOTE(nick): atomic qualified type. */
+  DW_TYPE_MDFR_ATOMIC = (1 << 0),
+
+  /* NOTE(nick): const qualified type. */
+  DW_TYPE_MDFR_CONST = (1 << 1),
+
+  /* NOTE(nick): Same as const, but the data lives in read-only part of the RAM. */
+  DW_TYPE_MDFR_IMMUTABLE = (1 << 2),
+  
+  /* NOTE(nick): When a data type is marked for compiler to use as little space as possible. */
+  DW_TYPE_MDFR_PACKED = (1 << 3),
+
+  /* NOTE(nick): A reference to (lvalue of) an object of the type being modifid. */
+  DW_TYPE_MDFR_REF = (1 << 4),
+
+  /* NOTE(nick): restrict qualified type. */
+  DW_TYPE_MDFR_RESTRICT = (1 << 5),
+
+  /* NOTE(nick): rvalue ref to an object of the type being modified (for example, in C++) */
+  DW_TYPE_MDFR_RVALUE_REF = (1 << 6),
+
+  /* NOTE(nick): shared qualified type. */
+  DW_TYPE_MDFR_SHARED = (1 << 7),
+
+  /* NOTE(nick): volatile qualified type. */
+  DW_TYPE_MDFR_VOLATILE = (1 << 8),
+
+  /* NOTE(nick): Assigned when memory values of this type can be displayed as text. */
+  DW_TYPE_MDFR_CHAR = (1 << 9),
+
+  DW_TYPE_MDFR_INVALID
+} DwTypeMdfr;
+
+typedef struct DwType 
+{
+  DwTag type_tag;
+
+  DwTag next_type_tag;
+
+  /* NOTE(nick): Type kind. */
+  DwTypeKind kind;
+
+  /* NOTE(nick): For more info look at DwTypeMdfr. */
+  U32 modifier;
+
+  /* NOTE(nick): Size of the type. */
+  U64 size;
+
+  /* NOTE(nick): Type name duh :) */
+  SymsString name;
+
+  /* NOTE(nick): Validity of the field is based on the "kind" of this type. */
+  union {
+    U64 arr_count;   /* NOTE(nick): DW_TYPE_ARR */
+    DwTag proc_params;  /* NOTE(nick): DW_TYPE_PROC, list of proc params. */
+    DwTag param_type;   /* NOTE(nick): DW_TYPE_PROC_PARAM, type of a param. */
+  } u;
+} DwType;
+
+typedef struct DwMember 
+{
+  /* NOTE(nick): Name of the a member. */
+  SymsString name;
+
+  /* NOTE(nick): Parent-relative offset of a member. */
+  U64 byte_off;
+
+  /* NOTE(nick): Type of a member. */
+  DwTag type_tag;
+} DwMember;
+
+typedef struct DwMemberIter 
+{
+  DwAttribIter attrib_iter;
+  U32 depth;
+} DwMemberIter;
+
+typedef enum 
+{
+  DW_LOCATION_NULL,
+  DW_LOCATION_ADDR,
+  DW_LOCATION_IMPLICIT
+} DwLocationType;
+
+typedef struct DwLocation 
+{
+  DwLocationType type;
+  union {
+    struct {
+      syms_uint len;
+      void *data;
+    } implicit;
+    SymsAddr addr;
+  } u;
+} DwLocation;
+
+typedef enum 
+{
+  DW_VAR_ARGUMENT = (1 << 0)
+} DwVarFlags;
+
+typedef struct DwVar 
+{
+  /* NOTE(nick): Name that user gave to the variable. */
+  SymsString name;
+
+  DwEncodedLocation encoded_va;
+
+  /* NOTE(nick): For more info look at DwVarFlag enum. */
+  U32 flags;
+  
+  /* NOTE(nick): You can pass this to the dw_infer_type to figure out type. */
+  DwTag type_tag;
+  
+  /* NOTE(nick): OPTIONAL! Line number where variable was declared. */
+  U32 decl_ln;
+  
+  /* NOTE(nick): OPTIONAL! Id of a file where variable lives. */
+  U32 decl_file;
+} DwVar;
+
+typedef DwVar DwLocal;
+typedef struct DwLocalIter 
+{
+  DwCompileUnit cu;
+  DwAttribIter attrib_iter;
+} DwLocalIter;
+
+typedef struct DwProcIter 
+{
+  DwAttribIter attribs;
+} DwProcIter;
+
+typedef struct DwFileIter 
+{
+  DwLineIter line_iter;
+  DwCuIter cu_iter;
+  DwCompileUnit cu;
+  syms_bool next_cu;
+} DwFileIter;
+
+typedef struct DwRangeIter
+{
+  DwBinRead rnglist;
+  SymsAddr base_addr;
+} DwRangeIter;
+
+typedef struct DwArangesUnit
+{
+  U64 unit_length;
+  U16 version;
+  SymsAddr debug_info_offset;
+  U8 addr_size;
+  U8 seg_size;
+  SymsAddr tuples_offset;
+  SymsAddr tuples_length;
+} DwArangesUnit;
+
+typedef struct DwArangesUnitIter
+{
+  DwContext *context;
+  DwBinRead aranges;
+} DwArangesUnitIter;
+
+typedef struct DwArangesIter
+{
+  DwBinRead aranges;
+  SymsAddr seg_size;
+} DwArangesIter;
+
+typedef struct 
+{
+  DwImgSec secs[DW_SEC_MAX];
+} DwInitdata;
+
+DW_API syms_bool
+dw_init(DwContext *context, SymsArch arch, DwInitdata *initdata);
+
+DW_API syms_bool
+dw_proc_iter_init(DwProcIter *proc_iter, DwCompileUnit *cu);
+
+DW_API syms_bool
+dw_proc_iter_next(DwProcIter *proc_iter, DwProc *proc);
+
+DW_API syms_bool
+dw_local_iter_init(DwLocalIter *local_iter, DwContext *context, DwTag proc_tag);
+
+DW_API syms_bool
+dw_local_iter_next(DwLocalIter *local_iter, DwLocal *local);
+
+DW_API syms_bool
+dw_member_iter_init(DwMemberIter *iter, DwContext *context, DwTag udt_tag);
+
+DW_API syms_bool
+dw_member_iter_next(DwMemberIter *member_iter, DwMember *member);
+
+DW_API syms_bool
+dw_infer_type(DwContext *context, DwTag type_tag, DwType *type);
+
+
+DW_API syms_bool
+dw_file_iter_init(DwFileIter *iter, DwContext *context);
+
+DW_API dw_uint
+dw_file_iter_next(DwFileIter *iter, void *bf, dw_uint bf_max);
+
+SYMS_INTERNAL syms_bool
+dw_cu_get_id_off(DwCompileUnit *cu, dw_uint id, SymsAddr *off);
+
+SYMS_INTERNAL syms_bool 
+dw_cu_init(DwCompileUnit *cu, DwContext *context, SymsAddr info_cu_base);
+
+SYMS_INTERNAL syms_bool
+dw_cu_iter_init(DwCuIter *iter, DwContext *context);
+
+SYMS_INTERNAL syms_bool
+dw_cu_iter_next(DwCuIter *iter, DwCompileUnit *cu);
+
+SYMS_INTERNAL DwEncodedLocation
+dw_encode_rva(DwCompileUnit *cu, U64 rva);
+
+DW_API syms_bool
+dw_decode_location(DwEncodedLocation *encoded_loc, SymsAddr rebase, void *memread_ctx, dw_memread_sig *memread, void *regread_ctx, dw_regread_sig *regread, DwLocation *decoded_loc);
+
+SYMS_INTERNAL syms_bool
+dw_encoded_location_is_valid(DwEncodedLocation loc);
+
+DW_API syms_bool
+dw_calc_heap_size(DwContext *context, SymsUMM *size);
+
+DW_API syms_bool
+dw_load_heap(DwContext *context, SymsArena *arena);
+
+SYMS_INTERNAL syms_bool
+dw_build_abbrev(DwCompileUnit *cu, SymsArena *arena);
+
+SYMS_INTERNAL syms_bool
+dw_class_init(DwAttribIter *attrib_iter, DwClass *udt);
+
+SYMS_INTERNAL syms_bool
+dw_scope_init(DwAttribIter *attrib_iter, DwScope *scope);
+
+SYMS_INTERNAL syms_bool
+dw_proc_init(DwAttribIter *attrib_iter, DwProc *proc);
+
+SYMS_INTERNAL syms_bool
+dw_cat_init(DwAttribIter *iter, DwCommonAttribs *cat);
+
+SYMS_INTERNAL syms_bool
+dw_var_init(DwAttribIter *attrib_iter, SymsAddr frame_base, SymsAddr member_location, SymsAddr cfa, DwVar *var);
+
+SYMS_INTERNAL syms_bool
+dw_name_iter_init(DwNameIter *iter, DwContext *context, DwNameTableIndex table_index);
+
+SYMS_INTERNAL syms_bool
+dw_name_iter_next(DwNameIter *iter, DwTag *tag, SymsString *name);
+
+SYMS_INTERNAL syms_bool
+dw_name_table_init(DwContext *context, DwNameTableIndex table_index, SymsArena *arena);
+
+SYMS_INTERNAL dw_uint
+dw_name_table_find(DwContext *context, DwNameTableIndex table_index, SymsString name, dw_uint tags_max, DwTag tags[]);
+
+DW_API syms_bool
+dw_abbrev_iter_init(DwAbbrevIter *iter, DwContext *context, SymsAddr abbrev_off);
+
+DW_API syms_bool
+dw_abbrev_iter_next(DwAbbrevIter *iter, DwAbbrevEntry *entry_out);
+
+/* TODO(nick): It's not good that we fetch an attribute with this procedure,
+ * which requires from user to be aware of this behaviour. Refactor! */
+DW_API syms_bool
+dw_attrib_iter_init(DwAttribIter *iter, DwCompileUnit *cu, SymsAddr info_off);
+
+DW_API syms_bool
+dw_attrib_iter_next(DwAttribIter *iter, DwAttrib *attrib);
+
+DW_API syms_bool
+dw_attrib_iter_reset(DwAttribIter *iter);
+
+SYMS_INTERNAL syms_bool
+dw_attrib_iter_next_tag(DwAttribIter *attrib_iter);
+
+SYMS_INTERNAL syms_bool
+dw_attrib_iter_skip_children(DwAttribIter *attrib_iter);
+
+DW_API SymsString
+dw_lang_to_str(DwLang lang);
+
+DW_API SymsString
+dw_tag_to_str(DwTagType tag);
+
+DW_API SymsString
+dw_at_to_str(DwAttribType at);
+
+DW_API SymsString
+dw_form_to_str(DwForm form);
+
+SYMS_INTERNAL syms_bool
+dw_range_check(DwCompileUnit *cu, SymsOffset range_off, SymsAddr addr, SymsAddr *lo_out, SymsAddr *hi_out);
+
+SYMS_INTERNAL syms_bool
+dw_range_iter_init(DwRangeIter *iter, DwCompileUnit *cu, SymsOffset range_off);
+
+SYMS_INTERNAL syms_bool
+dw_range_iter_next(DwRangeIter *iter, SymsAddr *lo_out, SymsAddr *hi_out);
+
+
+DW_API syms_bool
+dw_tag_iter_init(DwTagIter *iter, DwCompileUnit *cu, SymsOffset info_off);
+
+DW_API syms_bool
+dw_tag_iter_begin(DwTagIter *iter, dw_uint *depth_out, DwTagType *tag_out, DwAttribIter *attribs_out);
+
+DW_API syms_bool
+dw_tag_iter_next(DwTagIter *iter, DwAttribIter *attribs);
+
+DW_API DwAttribClass
+dw_pick_attrib_value_class(DwCompileUnit *cu, DwAttribType attrib, DwForm form);
+
+DW_API syms_bool
+dw_attrib_get_addr(DwAttrib *attrib, SymsAddr *addr);
+
+DW_API syms_bool
+dw_attrib_get_block(DwAttrib *attrib, DwFormBlock *block_out);
+
+DW_API syms_bool
+dw_attrib_get_const32(DwAttrib *attrib, U32 *value);
+
+DW_API syms_bool
+dw_attrib_get_const(DwAttrib *attrib, U64 *value);
+
+DW_API syms_bool
+dw_attrib_get_const128(DwAttrib *attrib, U64 *lo, U64 *hi);
+
+// Line table
+
+DW_API syms_bool
+dw_line_iter_init(DwLineIter *iter, DwCompileUnit *cu);
+
+DW_API syms_bool
+dw_line_iter_next(DwLineIter *iter, DwLineIterOp *op_out);
+
+DW_API syms_bool
+dw_line_iter_read_file(DwBinRead *linesec, DwLineFile *file_out);
+
+DW_API syms_bool
+dw_line_iter_read_dir(DwBinRead *linesec, SymsString *dir_out);
+
+SYMS_INTERNAL void
+dw_line_iter_advance_pc(DwLineIter *line_iter, U64 advance);
+
+SYMS_INTERNAL void
+dw_line_iter_reset_state(DwLineIter *iter);
+
+SYMS_INTERNAL syms_bool
+dw_line_iter_parse_file_and_add(DwLineIter *iter);
+
+SYMS_INTERNAL syms_bool
+dw_line_iter_parse_dir_and_add(DwLineIter *iter);
+
+SYMS_INTERNAL syms_bool
+dw_line_iter_update_filename(DwLineIter *iter, U64 file_index);
+
+SYMS_INTERNAL syms_bool
+dw_line_iter_test_filename(DwLineIter *iter, SymsString filename);
+
+SYMS_INTERNAL dw_uint
+dw_line_iter_get_filename(DwLineIter *iter, void *bf, dw_uint bf_max);
+
+//
+// String lookups
+//
+
+DW_API dw_uint
+dw_type_from_name(struct DwContext *context, const char *name, dw_uint name_len, dw_uint matches_max, DwTag *matches);
+
+DW_API syms_bool
+dw_global_from_name(struct DwContext *context, const char *name, dw_uint name_len, DwVar *var);
+
+#endif /* SYMS_INCLUDE_DWARF_H */
diff --git a/Engine/Source/Runtime/SymsLib/include/dwarf/syms_dwarf_api.h b/Engine/Source/Runtime/SymsLib/include/dwarf/syms_dwarf_api.h
new file mode 100644
--- /dev/null
+++ b/Engine/Source/Runtime/SymsLib/include/dwarf/syms_dwarf_api.h
@@ -0,0 +1,143 @@
+// Copyright Epic Games, Inc. All Rights Reserved.
+#ifndef SYMS_DWARF_API_INCLUDE_H
+#define SYMS_DWARF_API_INCLUDE_H
+
+/******************************************************************************
+ * File   : syms_dwarf.h                                                      *
+ * Author : Nikita Smith                                                      *
+ * Created: 2020/06/20                                                        *
+ * Purpose: Declarations for dwarf.h interface                                *
+ ******************************************************************************/
+
+typedef struct SymsDebugInfoDwarf 
+{
+  DwContext context;
+} SymsDebugInfoDwarf;
+
+typedef struct 
+{
+   SymsString compile_dir;
+   SymsString dir;
+   SymsString file;
+} DwFilePath;
+
+typedef struct SymsModIterDwarf 
+{
+  DwCuIter impl;
+} SymsModIterDwarf;
+
+typedef struct SymsLineIterDwarf 
+{
+  DwLineIter impl;
+  DwDirIndex prev_file_index;
+  DwCompileUnit *cu;
+} SymsLineIterDwarf;
+
+typedef struct SymsMemberIterDwarf 
+{
+  DwMemberIter impl;
+} SymsMemberIterDwarf;
+
+typedef struct SymsGlobalIterDwarf 
+{
+  U32 reserved;
+} SymsGlobalIterDwarf;
+
+typedef struct SymsLocalDataIterDwarf 
+{
+  U32 reserved;
+} SymsLocalDataIterDwarf;
+
+typedef struct SymsFileIterDwarf 
+{
+  DwFileIter impl;
+} SymsFileIterDwarf;
+
+typedef struct SymsProcIterDwarf 
+{
+  DwProcIter impl;
+} SymsProcIterDwarf;
+
+typedef struct SymsArgIterDwarf 
+{
+  U32 reserved;
+} SymsArgIterDwarf;
+
+typedef struct SymsLocalIterDwarf 
+{
+  DwLocalIter impl;
+} SymsLocalIterDwarf;
+
+typedef struct SymsInlineIterDwarf 
+{
+  SymsAddr rva;
+  DwCompileUnit *cu;
+  DwAttribIter attribs;
+  U32 depth;
+} SymsInlineIterDwarf;
+
+typedef struct SymsConstIterDwarf
+{
+  U32 reserved;
+} SymsConstIterDwarf;
+
+typedef struct SymsTypeIterDwarf 
+{
+  U32 reserved;
+} SymsTypeIterDwarf;
+
+typedef struct SymsRangeDwarf
+{
+  syms_uint cu_index;
+  SymsAddr range_off;
+} SymsRangeDwarf;
+
+typedef struct
+{
+  DwContext context;
+  DwCuIter cu_iter;
+  DwAttribIter att_iter;
+} SymsDebugFileIterDwarf;
+
+SYMS_COMPILER_ASSERT(sizeof(SymsUWord) >= 8);
+SYMS_COMPILER_ASSERT(sizeof(SymsUWord) >= 4);
+SYMS_COMPILER_ASSERT(SYMS_GET_MEMBER_SIZE(DwContext, secs) == SYMS_GET_MEMBER_SIZE(DwInitdata, secs));
+SYMS_COMPILER_ASSERT(sizeof(SymsRangeDwarf) <= sizeof(SymsRangeImpl));
+SYMS_COMPILER_ASSERT(sizeof(DwTag) <= sizeof(SymsTypeIDImpl));
+SYMS_COMPILER_ASSERT(sizeof(DwType) <= sizeof(SymsTypeImpl));
+SYMS_COMPILER_ASSERT(sizeof(SymsDebugInfoDwarf) <= sizeof(SymsDebugInfoImpl));
+SYMS_COMPILER_ASSERT(sizeof(DwFilePath) <= sizeof(SymsStringRefImpl));
+SYMS_COMPILER_ASSERT(sizeof(SymsModIterDwarf) <= sizeof(SymsModIterImpl));
+SYMS_COMPILER_ASSERT(sizeof(DwCompileUnit) <= sizeof(SymsModImpl));
+SYMS_COMPILER_ASSERT(sizeof(SymsLineIterDwarf) <= sizeof(SymsLineIterImpl));
+SYMS_COMPILER_ASSERT(sizeof(SymsMemberIterDwarf) <= sizeof(SymsMemberIterImpl));
+SYMS_COMPILER_ASSERT(sizeof(SymsGlobalIterDwarf) <= sizeof(SymsGlobalIterImpl));
+SYMS_COMPILER_ASSERT(sizeof(SymsLocalDataIterDwarf) <= sizeof(SymsLocalDataIterImpl));
+SYMS_COMPILER_ASSERT(sizeof(SymsFileIterDwarf) <= sizeof(SymsFileIterImpl));
+SYMS_COMPILER_ASSERT(sizeof(SymsArgIterDwarf) <= sizeof(SymsArgIterImpl));
+SYMS_COMPILER_ASSERT(sizeof(SymsLocalIterDwarf) <= sizeof(SymsLocalIterImpl));
+SYMS_COMPILER_ASSERT(sizeof(SymsInlineIterDwarf) <= sizeof(SymsInlineIterImpl));
+SYMS_COMPILER_ASSERT(sizeof(SymsDebugFileIterDwarf) <= sizeof(SymsDebugFileIterImpl));
+typedef DW_MEMREAD_SIG(SymsDwMemreadSig);
+typedef DW_REGREAD_SIG(SymsDwRegreadSig);
+typedef DW_REGWRITE_SIG(SymsDwRegwriteSig);
+
+// -----------------------------------------------------------------------------
+
+SYMS_INTERNAL
+DW_REGREAD_SIG(syms_regread_dwarf);
+
+SYMS_INTERNAL
+DW_MEMREAD_SIG(syms_memread_dwarf);
+
+SYMS_INTERNAL syms_bool
+syms_dw_type_to_syms_type(DwType *dw_type, SymsType *syms_type);
+
+SYMS_INTERNAL syms_bool
+syms_proc_from_dw_proc(SymsDebugInfoDwarf *debug_info, DwCompileUnit *cu, DwProc *proc, SymsProc *proc_out);
+
+SYMS_INTERNAL SymsStringRef
+syms_string_ref_dw_path(SymsString compile_dir, SymsString dir, SymsString file);
+
+#endif /* SYMS_DWARF_API_INCLUDE_H */
+
diff --git a/Engine/Source/Runtime/SymsLib/include/dwarf/syms_dwarf_unwind.h b/Engine/Source/Runtime/SymsLib/include/dwarf/syms_dwarf_unwind.h
new file mode 100644
--- /dev/null
+++ b/Engine/Source/Runtime/SymsLib/include/dwarf/syms_dwarf_unwind.h
@@ -0,0 +1,253 @@
+// Copyright Epic Games, Inc. All Rights Reserved.
+#ifndef SYMS_DWARF_UNWIND_INCLUDE_H
+#define SYMS_DWARF_UNWIND_INCLUDE_H
+
+#ifndef DW_API
+#define DW_API SYMS_API
+#endif
+
+typedef struct DwArchInfo 
+{
+  DwMode mode;
+  SymsArch arch;
+  dw_uint word_size;
+  dw_uint addr_size;
+  dw_uint ip_regid;
+  dw_uint sp_regid;
+} DwArchInfo; 
+
+typedef struct DwUserCallbacks 
+{
+  DwArchInfo arch_info;
+
+  void *regread_ctx;
+  dw_regread_sig *regread;
+
+  void *regwrite_ctx;
+  dw_regwrite_sig *regwrite;
+
+  void *memread_ctx;
+  dw_memread_sig *memread;
+} DwUserCallbacks;
+
+#define DW_PRESERVED_REGS_COUNT_NT      33
+#define DW_PRESERVED_REGS_COUNT_UNIX    17
+#define DW_PRESERVED_REGS_MAX           (SYMS_MAX(DW_PRESERVED_REGS_COUNT_NT, DW_PRESERVED_REGS_COUNT_UNIX))
+
+#define DW_CFA_OPCODE_MASK 0xC0
+#define DW_CFA_OPER_MASK   0x3F
+typedef enum 
+{
+  DW_CFA_NOP                   = 0x0,
+  DW_CFA_SET_LOC               = 0x1,
+  DW_CFA_ADVANCE_LOC1          = 0x2,
+  DW_CFA_ADVANCE_LOC2          = 0x3,
+  DW_CFA_ADVANCE_LOC4          = 0x4,
+  DW_CFA_OFFSET_EXT            = 0x5,
+  DW_CFA_RESTORE_EXT           = 0x6,
+  DW_CFA_UNDEFINED             = 0x07,
+  DW_CFA_SAME_VALUE            = 0x8,
+  DW_CFA_REGISTER              = 0x9,
+  DW_CFA_REMEMBER_STATE        = 0xA,
+  DW_CFA_RESTORE_STATE         = 0xB,
+  DW_CFA_DEF_CFA               = 0xC,
+  DW_CFA_DEF_CFA_REGISTER      = 0xD,
+  DW_CFA_DEF_CFA_OFFSET        = 0xE,
+  DW_CFA_DEF_CFA_EXPR          = 0xF,
+  DW_CFA_EXPR                  = 0x10,
+  DW_CFA_OFFSET_EXT_SF         = 0x11,
+  DW_CFA_DEF_CFA_SF            = 0x12,
+  DW_CFA_DEF_CFA_OFFSET_SF     = 0x13,
+  DW_CFA_VAL_OFFSET            = 0x14,
+  DW_CFA_VAL_OFFSET_SF         = 0x15,
+  DW_CFA_VAL_EXPR              = 0x16,
+
+  DW_CFA_ADVANCE_LOC = 0x40,
+  DW_CFA_OFFSET      = 0x80,
+  DW_CFA_RESTORE     = 0xC0,
+
+  DW_CFA_USER_LO = 0x1C,
+  DW_CFA_USER_HI = 0x3F
+} DwCFAOpcodeType;
+
+#define DW_EH_PE_FORMAT_MASK    0x0F
+#define DW_EH_PE_APPLY_MASK     0x70
+#define DW_EH_PE_INDIRECT       0x80
+
+#define DW_EH_PE_OMIT       0xFFu
+/* Pointer sized unsigned value */
+#define DW_EH_PE_PTR        0x00u
+/* Unsigned LE base-128 value */
+#define DW_EH_PE_ULEB128    0x01u
+/* Unsigned 16-bit value */
+#define DW_EH_PE_UDATA2     0x02u
+/* Unsigned 32-bit value */
+#define DW_EH_PE_UDATA4     0x03u
+/* Unsigned 64-bit value */
+#define DW_EH_PE_UDATA8     0x04u
+/* Signed pointer */
+#define DW_EH_PE_SIGNED     0x08u
+/* Signed LE base-128 value */
+#define DW_EH_PE_SLEB128    0x09u
+/* Signed 16-bit value */
+#define DW_EH_PE_SDATA2     0x0Au
+/* Signed 32-bit value */
+#define DW_EH_PE_SDATA4     0x0Bu
+/* Signed 64-bit value */
+#define DW_EH_PE_SDATA8     0x0Cu
+typedef dw_uint DwEhPointerEncodingFormat;
+
+typedef enum 
+{
+  DW_EH_PE_ABSPTR  = 0x00, /* Absolute value */
+  DW_EH_PE_PCREL   = 0x10, /* Relative to address of encoded value */
+  DW_EH_PE_TEXTREL = 0x20, /* Text-Relative (GCC-specific) */
+  DW_EH_PE_DATAREL = 0x30, /* Data-Relative */
+  /* The following are not documented by LSB v1.3, yet they are used by
+     GCC, presumably they aren't documented by LSB since they aren't
+     used on Linux:  */
+  DW_EH_PE_FUNCREL = 0x40, /* Start of procedure relative */
+  DW_EH_PE_ALIGNED = 0x50  /* Aligned Pointer */
+} DwEHPointerEncodingApplication;
+
+typedef enum 
+{
+  DW_VIRTUAL_UNWIND_DATA_NULL,
+  DW_VIRTUAL_UNWIND_DATA_EH_FRAME,
+  DW_VIRTUAL_UNWIND_DATA_DEBUG_FRAME
+} DwVirtualUnwindDataType;
+
+#define DW_CIE_VERSION      3
+#define DW_CIE_VERSION_MAX  4
+
+typedef struct DwCommonInfoEntry 
+{
+  SymsAddr init_cfi_offset;
+  SymsAddr init_cfi_size;
+  SymsAddr end_offset;
+
+  SymsUWord cie_id;
+  SymsString augmentation;
+
+  syms_bool is_aug_sized;
+  syms_bool is_sig_frame;
+  syms_bool have_abi_maker;
+
+  U8 version;
+  U8 fde_encoding;
+  U8 lsda_encoding;
+  U8 handler_encoding;
+
+  S64 addr_size;
+  S64 segsel_size;
+
+  S64 code_align;
+  U64 data_align;
+  U64 ret_addr_reg;
+
+  SymsAddr handler_ip;
+} DwCommonInfoEntry;
+
+typedef enum DwCfiRegisterType 
+{
+  DW_CFI_REGISTER_TYPE_UNDEF,
+  DW_CFI_REGISTER_TYPE_SAME,
+  DW_CFI_REGISTER_TYPE_CFAREL,
+  DW_CFI_REGISTER_TYPE_REG,
+  DW_CFI_REGISTER_TYPE_EXPR,
+  DW_CFI_REGISTER_TYPE_VAL_EXPR,
+  DW_CFI_REGISTER_TYPE_INVALID
+} DwCfiRegisterType;
+
+typedef union DwCfiRegValue 
+{
+  SymsSWord w;
+  DwEncodedLocationExpr e;
+} DwCfiRegValue;
+
+static const dw_uint DW_CFA_COLUMN_REG = (DW_PRESERVED_REGS_MAX + 0);
+static const dw_uint DW_CFA_COLUMN_OFF = (DW_PRESERVED_REGS_MAX + 1);
+
+typedef struct DwCfiRow 
+{
+  DwCfiRegisterType type[DW_PRESERVED_REGS_MAX + 2];
+  DwCfiRegValue value[DW_PRESERVED_REGS_MAX + 2];
+} DwCfiRow;
+
+#define DW_CFI_PROGRAM_STACK_MAX 128
+typedef struct DwCfiProgram 
+{
+  DwCfiRow rules;
+
+  DwCfiRow *frame;
+  DwCfiRow stack[DW_CFI_PROGRAM_STACK_MAX];
+
+  dw_uint reg_count;
+  syms_bool setup_cfa;
+  S64 ret_addr_regid;
+  SymsAddr cfa;
+} DwCfiProgram;
+
+typedef struct DwFrameInfo 
+{
+  DwVirtualUnwindDataType source_type;
+  SymsAddr eh_frame;
+  SymsAddr image_base;
+  DwArchInfo arch_info;
+} DwFrameInfo;
+
+typedef struct DwVirtualUnwind 
+{
+  DwCfiProgram program;
+} DwVirtualUnwind;
+
+typedef struct DwFrameDescEntry 
+{
+  /* NOTE(nick): Offset where FDE was read */
+  SymsAddr data_off;
+
+  /* NOTE(nick): Offset of CFI for this FDE */
+  SymsAddr cfi_offset;
+
+  /* NOTE(nick): Number of bytes that FDE occupies */
+  U64 data_size;
+
+  /* NOTE(nick): Number of bytes that CFI take up */
+  U64 cfi_size;
+
+  /* NOTE(nick): Address of first instruction in procedure, that FDE describes. */
+  SymsAddr start_ip;
+
+  /* NOTE(nick): Number of bytes that instructions occupy. */
+  SymsAddr range_ip;
+
+  /* NOTE(nick): Address of LSDA. */
+  SymsAddr lsda_ip;
+} DwFrameDescEntry;
+
+typedef struct DwFrameDescEntryIter 
+{
+  DwFrameInfo *frame_info;
+  DwBinRead secdata;
+  DwCommonInfoEntry cie;
+  SymsAddr cie_offset;
+} DwFrameDescEntryIter;
+
+DW_API syms_bool
+dw_virtual_unwind_init(SymsImageType image_type, DwVirtualUnwind *context_out);
+
+DW_API syms_bool
+dw_virtual_unwind_frame(DwVirtualUnwind *context,
+            SymsArch arch,
+            DwVirtualUnwindDataType sec_bytes_type,
+            void *sec_bytes, SymsUMM sec_bytes_size,
+            SymsAddr image_base,
+            SymsAddr sec_bytes_base,
+            void *memread_ctx, dw_memread_sig *memread,
+            void *regread_ctx, dw_regread_sig *regread,
+            void *regwrite_ctx, dw_regwrite_sig *regwrite);
+
+SYMS_INTERNAL SymsAddr
+dw_parse_pointer(DwFrameInfo *frame_info, DwBinRead *secdata, dw_uint encoding);
+
+#endif /* SYMS_DWARF_UNWIND_INCLUDE_H */
diff --git a/Engine/Source/Runtime/SymsLib/include/elf/syms_elf.h b/Engine/Source/Runtime/SymsLib/include/elf/syms_elf.h
new file mode 100644
--- /dev/null
+++ b/Engine/Source/Runtime/SymsLib/include/elf/syms_elf.h
@@ -0,0 +1,421 @@
+// Copyright Epic Games, Inc. All Rights Reserved.
+#ifndef SYMS_LINUX_INCLUDE_H
+#define SYMS_LINUX_INCLUDE_H
+
+#ifndef ELF_API
+#define ELF_API SYMS_API
+#endif
+
+//  File   : syms_elf.h
+//  Author : Nikita Smith
+//  Created: 2020/05/30
+//  Purpose: Declarations of ELF data structs and rountines
+
+typedef U32 Syms_Elf32_Addr;
+typedef U16 Syms_Elf32_Half;
+typedef U32 Syms_Elf32_Off;
+typedef S32 Syms_Elf32_Sword;
+typedef U32 Syms_Elf32_Word;
+
+typedef U64 Syms_Elf64_Addr;
+typedef U16 Syms_Elf64_Half;
+typedef S16 Syms_Elf64_SHalf;
+typedef U64 Syms_Elf64_Off;
+typedef S32 Syms_Elf64_Sword;
+typedef U32 Syms_Elf64_Word;
+typedef U64 Syms_Elf64_Xword;
+typedef S64 Syms_Elf64_Sxword;
+
+#define SYMS_ELF32_ST_BIND(i) ((i) >> 4)
+#define SYMS_ELF32_ST_TYPE(i) ((i) & 0xF)
+#define SYMS_ELF32_INFO(b, t) (((b) << 4) + ((t) & 0xf))
+
+#define SYMS_ELF64_ST_BIND(i)    ((i) >> 4)
+#define SYMS_ELF64_ST_TYPE(i)    ((i) & 0xf)
+#define SYMS_ELF64_ST_INFO(b, t) (((b) << 4) + ((t) & 0xf))
+
+#define SYMS_STN_UNDEF 0
+
+#define SYMS_STB_LOCAL      0
+#define SYMS_STB_GLOBAL     1
+#define SYMS_STB_WEAK       2
+#define SYMS_STB_LOOS      10
+#define SYMS_STB_HIOS      12
+#define SYMS_STB_LOPROC    13
+#define SYMS_STB_HIPROC    15
+
+#define SYMS_STT_NOTYPE     0
+#define SYMS_STT_OBJECT     1
+#define SYMS_STT_FUNC       2
+#define SYMS_STT_SECTION    3
+#define SYMS_STT_FILE       4
+#define SYMS_STT_COMMON     5
+#define SYMS_STT_TLS        6
+#define SYMS_STT_LOOS      10
+#define SYMS_STT_HIOS      12
+#define SYMS_STT_LOPROC    13
+#define SYMS_STT_HIPROC    15
+
+#define SYMS_ELF32_ST_VISIBILITY(o) ((o) & 0x3)
+#define SYMS_ELF64_ST_VISIBILITY(o) ((o) & 0x3)
+
+#define SYMS_STV_DEFAULT    0
+#define SYMS_STV_INTERNAL   1
+#define SYMS_STV_HIDDEN     2
+#define SYMS_STV_PROTECTED  3
+
+// p_type
+#define SYMS_PT_NONE        0u
+#define SYMS_PT_LOAD        1u
+#define SYMS_PT_DYNAMIC     2u
+#define SYMS_PT_INTERP      3u
+#define SYMS_PT_NOTE        4u
+#define SYMS_PT_SHLIB       5u
+#define SYMS_PT_PHDR        6u
+#define SYMS_PT_LOPROC      0x70000000u
+#define SYMS_PT_HIPROC      0x7fffffffu
+// Specific to Sun
+#define SYMS_PT_LOSUNW      0x6ffffffau
+#define SYMS_PT_SUNWBSS     0x6ffffffbu
+#define SYMS_PT_SUNWSTACK   0x6ffffffau
+#define SYMS_PT_HISUNW      0x6fffffffu
+
+// e_machine
+#define SYMS_EM_NONE           0u
+#define SYMS_EM_M32            1u
+#define SYMS_EM_SPARC          2u
+#define SYMS_EM_386            3u
+#define SYMS_EM_68K            4u
+#define SYMS_EM_88K            5u
+#define SYMS_EM_IAMCU          6u
+#define SYMS_EM_860            7u
+#define SYMS_EM_MIPS           8u
+#define SYMS_EM_S370           9u
+#define SYMS_EM_MIPS_RS3_LE    10u
+#define SYMS_EM_PARISC         15u
+#define SYMS_EM_PPC_OLD        17u /* Old version of PowerPC. Deprecated */
+#define SYMS_EM_SPARC32PLUS    18u /* Sun's "v8plus" */
+#define SYMS_EM_960            19u
+#define SYMS_EM_PPC            20u
+#define SYMS_EM_PPC64          21u
+#define SYMS_EM_S390           22u
+#define SYMS_EM_SPU            23u
+#define SYMS_EM_V800           36u
+#define SYMS_EM_FR20           37u
+#define SYMS_EM_RH32           38u
+#define SYMS_EM_MCORE          39u
+#define SYMS_EM_ARM            40u
+#define SYMS_EM_SH             42u
+#define SYMS_EM_IA_64          50u
+#define SYMS_EM_X86_64         62u
+#define SYMS_EM_AARCH64        183u
+#define SYMS_EM_RISCV          243u
+
+#define SYMS_ELFOSABI_NONE       0u
+#define SYMS_ELFOSABI_SYSV       0u
+#define SYMS_ELFOSABI_HPUX       1u
+#define SYMS_ELFOSABI_NETBSD     2u
+#define SYMS_ELFOSABI_GNU        3u
+#define SYMS_ELFOSABI_HURD       4u
+#define SYMS_ELFOSABI_86OPEN     5u
+#define SYMS_ELFOSABI_SOLARIS    6u
+#define SYMS_ELFOSABI_AIX        7u
+#define SYMS_ELFOSABI_IRIX       8u
+#define SYMS_ELFOSABI_FREEBSD    9u
+#define SYMS_ELFOSABI_TRU64      10u
+#define SYMS_ELFOSABI_MODESTO    11u
+#define SYMS_ELFOSABI_OPENBSD    12u
+#define SYMS_ELFOSABI_OPENVMS    13u
+#define SYMS_ELFOSABI_NSK	       14u
+#define SYMS_ELFOSABI_AROS	     15u
+#define SYMS_ELFOSABI_FENIXOS    16u
+#define SYMS_ELFOSABI_CLOUDABI   17u
+#define SYMS_ELFOSABI_OPENVOS    18u
+#define SYMS_ELFOSABI_ARM_FDPIC  65u
+#define SYMS_ELFOSABI_ARM        97u
+#define SYMS_ELFOSABI_STANDALONE 255u
+
+/* sh_type */
+#define SYMS_SHT_NULL     0
+#define SYMS_SHT_PROGBITS 1
+#define SYMS_SHT_SYMTAB   2
+#define SYMS_SHT_STRTAB   3
+#define SYMS_SHT_RELA     4
+#define SYMS_SHT_HASH     5
+#define SYMS_SHT_DYNAMIC  6
+#define SYMS_SHT_NOTE     7
+#define SYMS_SHT_NOBITS   8
+#define SYMS_SHT_REL      9
+#define SYMS_SHT_SHLIB    10
+#define SYMS_SHT_DYNSYM   11
+#define SYMS_SHT_INIT_ARRAY	    14		/* Array of ptrs to init functions */
+#define SYMS_SHT_FINI_ARRAY	    15		/* Array of ptrs to finish functions */
+#define SYMS_SHT_PREINIT_ARRAY  16		/* Array of ptrs to pre-init funcs */
+#define SYMS_SHT_GROUP	        17		/* Section contains a section group */
+#define SYMS_SHT_SYMTAB_SHNDX   18		/* Indices for SHN_XINDEX entries */
+
+#define SYMS_SHT_GNU_INCREMENTAL_INPUTS 0x6fff4700   /* incremental build data */
+#define SYMS_SHT_GNU_ATTRIBUTES 0x6ffffff5	/* Object attributes */
+#define SYMS_SHT_GNU_HASH	0x6ffffff6	/* GNU style symbol hash table */
+#define SYMS_SHT_GNU_LIBLIST	0x6ffffff7	/* List of prelink dependencies */
+
+/* The next three section types are defined by Solaris, and are named
+   SHT_SUNW*.  We use them in GNU code, so we also define SHT_GNU*
+   versions.  */
+#define SYMS_SHT_SUNW_verdef	0x6ffffffd	/* Versions defined by file */
+#define SYMS_SHT_SUNW_verneed 0x6ffffffe	/* Versions needed by file */
+#define SYMS_SHT_SUNW_versym	0x6fffffff	/* Symbol versions */
+
+#define SYMS_SHT_GNU_verdef	  SYMS_SHT_SUNW_verdef
+#define SYMS_SHT_GNU_verneed	SYMS_SHT_SUNW_verneed
+#define SYMS_SHT_GNU_versym	  SYMS_SHT_SUNW_versym
+
+#define SYMS_SHT_LOPROC   0x70000000
+#define SYMS_SHT_HIPROC   0x7fffffff
+#define SYMS_SHT_LOUSER   0x80000000
+#define SYMS_SHT_HIUSER   0xffffffff
+
+#define SYMS_SHN_UNDEF      0
+#define SYMS_SHN_LORESERVE  0xff00
+#define SYMS_SHN_LOPROC     0xff00
+#define SYMS_SHN_BEFORE     0xff00
+#define SYMS_SHN_AFTER      0xff01
+#define SYMS_SHN_HIPROC     0xff1f
+#define SYMS_SHN_ABS        0xfff1
+#define SYMS_SHN_COMMON     0xfff2
+#define SYMS_SHN_HIRESERVE  0xffff
+
+#define SYMS_SHF_WRITE            0x1ul
+#define SYMS_SHF_ALLOC            0x2ul
+#define SYMS_SHF_EXECINSTR        0x4ul
+#define SYMS_SHF_MERGE            0x10ul
+#define SYMS_SHF_STRINGS          0x20ul
+#define SYMS_SHF_INFO_LINK        0x40ul
+#define SYMS_SHF_LINK_ORDER       0x80ul
+#define SYMS_SHF_OS_NONCONFORMING 0x100ul
+#define SYMS_SHF_GROUP            0x200ul
+#define SYMS_SHF_TLS              0x400ul
+#define SYMS_SHF_COMPRESSED       0x800ul
+#define SYMS_SHF_ORDERED          0x4000000ul
+#define SYMS_SHF_EXCLUDE          0x8000000ul
+#define SYMS_SHF_GNU_BUILD_NOTE   0x100000ul
+#define SYMS_SHF_GNU_MBIND        0x01000000ul
+
+#define SYMS_SHF_MASKOS           0x0ff00000ul
+#define SYMS_SHF_MASKPROC         0xf0000000ul
+
+#define SYMS_ELFCLASSNONE 0
+#define SYMS_ELFCLASS32   1
+#define SYMS_ELFCLASS64   2
+
+#define SYMS_ET_NONE    0u
+// NOTE(nick): https://refspecs.linuxfoundation.org/elf/elf.pdf (page 15)
+//
+// File holds code and data for linking with other object files to 
+// create an executable or a shared object file.
+#define SYMS_ET_REL     1u
+// File holds a program suitable for execution.
+#define SYMS_ET_EXEC    2u
+// File holds code and data suitable for linking in two contexts. First, the 
+// link editor may process it with other relocatable and shared object file to
+// create another object file. Second, the dynamic linker combines it with
+// with an executable file and other shared objects to create a process image.
+#define SYMS_ET_DYN     3u
+// Contents of file is not specified (it's used for processor specific data).
+#define SYMS_ET_CORE    4u
+#define SYMS_ET_LOPROC  0xff00u
+#define SYMS_ET_HIPROC  0xffffu
+
+#if 0 // OribsOS flags for e_type
+#define SYMS_SCE_EXEC    0xfe10u
+#define SYMS_SCE_DYNEXEC 0xfe00u
+#endif
+
+typedef struct 
+{
+  Syms_Elf32_Word st_name;
+  Syms_Elf32_Addr st_value;
+  Syms_Elf32_Word st_size;
+  unsigned char st_info;
+  unsigned char st_other;
+  Syms_Elf32_Half st_shndx;
+} SymsElf32Sym;
+
+typedef struct 
+{
+  Syms_Elf64_Word st_name;
+  unsigned char st_info;
+  unsigned char st_other;
+  Syms_Elf64_Half st_shndx;
+  Syms_Elf64_Addr st_value;
+  Syms_Elf64_Xword st_size;
+} SymsElf64Sym;
+
+typedef struct SymsElfShdr32 
+{
+  U32 sh_name;
+  U32 sh_type;
+  U32 sh_flags;
+  U32 sh_addr;
+  U32 sh_offset;
+  U32 sh_size;
+  U32 sh_link;
+  U32 sh_info;
+  U32 sh_addralign;
+  U32 sh_entsize;
+} SymsElfShdr32;
+
+typedef struct SymsElfShdr64 
+{
+  U32 sh_name;		// Section name, index in string tbl 
+  U32 sh_type;		// Type of section 
+  U64 sh_flags;		// Miscellaneous section attributes 
+  U64 sh_addr;		// Section virtual addr at execution 
+  U64 sh_offset;	// Section file offset 
+  U64 sh_size;		// Size of section in bytes 
+  U32 sh_link;		// Index of another section 
+  U32 sh_info;		// Additional section information 
+  U64 sh_addralign;	// Section alignment 
+  U64 sh_entsize;	// Entry size if section holds table 
+} SymsElfShdr64;
+
+typedef struct SymsElfPhdr32 
+{
+  U32	p_type;
+  U32	p_offset;
+  U32	p_vaddr;
+  U32	p_paddr;
+  U32	p_filesz;
+  U32	p_memsz;
+  U32	p_flags;
+  U32	p_align;
+} SymsElfPhdr32;
+
+typedef struct SymsElfPhdr64 
+{
+  U32 p_type;
+  U32 p_flags;
+  U64 p_offset;   // Segment file offset 
+  U64 p_vaddr;    // Segment virtual address 
+  U64 p_paddr;    // Segment physical address 
+  U64 p_filesz;   // Segment size in file 
+  U64 p_memsz;    // Segment size in memory 
+  U64 p_align;    // Segment alignment, file & memory
+} SymsElfPhdr64;
+
+#define SYMS_EI_MAG0    0
+#define SYMS_EI_MAG1    1
+#define SYMS_EI_MAG2    2
+#define SYMS_EI_MAG3    3
+#define SYMS_EI_CLASS   4
+#define SYMS_EI_DATA    5
+#define SYMS_EI_VERSION 6
+#define SYMS_EI_PAD     7
+#define SYMS_EI_NIDENT  16
+
+#define SYMS_ELF_SIG    0x464c457f
+
+#define SYMS_EV_NONE    0
+#define SYMS_EV_CURRENT 1
+
+typedef struct SymsElf32 
+{
+  U8 e_ident[SYMS_EI_NIDENT];
+  U16 e_type;
+  U16 e_machine;
+  U16 e_version;
+  Syms_Elf32_Addr e_entry;
+  Syms_Elf32_Off e_phoff;
+  Syms_Elf32_Off e_shoff;
+  U32 e_flags;
+  U16 e_ehsize;
+  U16 e_phentsize;
+  U16 e_phnum;
+  U16 e_shentsize;
+  U16 e_shnum;
+  U16 e_shstrndx;
+} SymsElf32;
+
+typedef struct SymsElf64 
+{
+  U8 e_ident[SYMS_EI_NIDENT];
+  U16 e_type;
+  U16 e_machine;
+  U16 e_version;
+  Syms_Elf64_Addr e_entry;
+  Syms_Elf64_Off e_phoff;
+  Syms_Elf64_Off e_shoff;
+  U32 e_flags;
+  U16 e_ehsize;
+  U16 e_phentsize;
+  U16 e_phnum;
+  U16 e_shentsize;
+  U16 e_shnum;
+  U16 e_shstrndx;
+} SymsElf64;
+
+// ----------------------------------------------------------------------------
+
+typedef struct SymsImageElf 
+{
+  SymsAddr sh_name_lo;
+  SymsAddr sh_name_hi;
+  union {
+    SymsElf32 *header32;
+    SymsElf64 *header64;
+  } u;
+} SymsImageElf;
+
+typedef struct SymsSecIterElf 
+{
+  struct SymsImage *image;
+  SymsAddr sh_name_lo;
+  SymsAddr sh_name_hi;
+  U32 header_index;
+  U32 header_count;
+  SymsBuffer headers;
+} SymsSecIterElf;
+
+typedef struct SymsSymtabIter 
+{
+  U32 index;
+  U32 count;
+  SymsBuffer symtab_cursor;
+  SymsBuffer strtab_cursor;
+  SymsImageHeaderClass header_class;
+} SymsSymtabIter;
+SYMS_COMPILER_ASSERT(sizeof(SymsSymtabIter) <= sizeof(SymsModImpl));
+
+typedef struct SymsSymtabEntry 
+{
+  SymsString name;
+  SymsAddr value;
+  SymsUWord size;
+  U32 index;
+  U8 bind;
+  U8 type;
+  U8 vis;
+} SymsSymtabEntry;
+
+SYMS_API syms_bool
+syms_img_init_elf(struct SymsImage *img, void *img_data, SymsUWord img_size, SymsLoadImageFlags flags);
+
+SYMS_INTERNAL SymsSecIterElf
+syms_sec_iter_init_elf(struct SymsImage *img);
+
+SYMS_INTERNAL syms_bool
+syms_sec_iter_next_elf(SymsSecIterElf *iter, SymsElfShdr64 *sec_out);
+
+SYMS_INTERNAL struct SymsProc
+syms_proc_from_stt_func(SymsSymtabEntry *stt_func);
+
+ELF_API syms_bool
+syms_symtab_iter_init(struct SymsInstance *instance, SymsSymtabIter *iter);
+
+ELF_API syms_bool
+syms_symtab_iter_next(SymsSymtabIter *iter, SymsSymtabEntry *entry_out);
+
+SYMS_INTERNAL SymsAddr
+syms_get_rebase_elf(const SymsImageElf *elf, SymsImageHeaderClass header_class, SymsAddr old_base, SymsAddr base);
+
+#endif
