RAD symbol library

diff --git a/Engine/Source/Runtime/SymsLib/include/pdb/syms_codeview.h b/Engine/Source/Runtime/SymsLib/include/pdb/syms_codeview.h
new file mode 100644
--- /dev/null
+++ b/Engine/Source/Runtime/SymsLib/include/pdb/syms_codeview.h
@@ -0,0 +1,2352 @@
+// Copyright Epic Games, Inc. All Rights Reserved.
+#ifndef SYMS_CODEVIEW_INCLUDE_H
+#define SYMS_CODEVIEW_INCLUDE_H
+
+#pragma pack(push, 1)
+
+/* NOTE(nick; Oct 5 2019): There are two type maps stored in the PDB.
+ *
+ * First lives in the PDB_DEFAULT_STREAM_TPI and any index
+ * of type pdb_itype is related to this stream.
+ *
+ * Second lives in the PDB_DEFAULT_STREAM_IPI and any index
+ * of type pdb_cv_itemid is related to this stream. It contains
+ * only following records:
+ *  PDB_LF_FUNC_ID,
+ *  PDB_LF_MFUNC_ID,
+ *  PDB_LF_BUILDINFO,
+ *  PDB_LF_SUBSTR_LIST,
+ *  PDB_LF_STRING_ID
+ *
+ * */
+
+#if defined(PDB_COMPILE_64BIT)
+  #define PDB_INVALID_TI SYMS_UINT64_MAX
+  typedef U64 pdb_ti;
+#else
+  #define PDB_INVALID_TI SYMS_UINT32_MAX
+  typedef U32 pdb_ti;
+#endif
+
+#define PDB_INVALID_ITYPE   PDB_INVALID_TI
+#define PDB_ITYPE_VARIADIC (PDB_INVALID_ITYPE - 1)
+typedef pdb_ti pdb_cv_itype;
+
+#define PDB_INVALID_ITEMID PDB_INVLAID_TI
+typedef pdb_ti pdb_cv_itemid;
+
+/* NOTE(nick; Oct 5 2019): Module Index */
+typedef U16 pdb_imod;
+#define PDB_CV_INVALID_IMOD SYMS_UINT16_MAX
+
+/* NOTE(nick; Oct 5 2019): Section Index */
+typedef U16 pdb_isec;
+typedef U32 pdb_isec_umm;
+
+typedef enum {
+  PDB_CV_X86_NONE     =   0, PDB_CV_X86_EAX      =  17,
+  PDB_CV_X86_AL       =   1, PDB_CV_X86_ECX      =  18,
+  PDB_CV_X86_CL       =   2, PDB_CV_X86_EDX      =  19,
+  PDB_CV_X86_DL       =   3, PDB_CV_X86_EBX      =  20,
+  PDB_CV_X86_BL       =   4, PDB_CV_X86_ESP      =  21,
+  PDB_CV_X86_AH       =   5, PDB_CV_X86_EBP      =  22,
+  PDB_CV_X86_CH       =   6, PDB_CV_X86_ESI      =  23,
+  PDB_CV_X86_DH       =   7, PDB_CV_X86_EDI      =  24,
+  PDB_CV_X86_BH       =   8, PDB_CV_X86_ES       =  25,
+  PDB_CV_X86_AX       =   9, PDB_CV_X86_CS       =  26,
+  PDB_CV_X86_CX       =  10, PDB_CV_X86_SS       =  27,
+  PDB_CV_X86_DX       =  11, PDB_CV_X86_DS       =  28,
+  PDB_CV_X86_BX       =  12, PDB_CV_X86_FS       =  29,
+  PDB_CV_X86_SP       =  13, PDB_CV_X86_GS       =  30,
+  PDB_CV_X86_BP       =  14, PDB_CV_X86_IP       =  31,
+  PDB_CV_X86_SI       =  15, PDB_CV_X86_FLAGS    =  32,
+  PDB_CV_X86_DI       =  16, PDB_CV_X86_EIP      =  33,
+  PDB_CV_X86_EFLAGS   =  34,
+
+  PDB_CV_X86_MM0      =  146, PDB_CV_X86_XMM0     =  154,
+  PDB_CV_X86_MM1      =  147, PDB_CV_X86_XMM1     =  155,
+  PDB_CV_X86_MM2      =  148, PDB_CV_X86_XMM2     =  156,
+  PDB_CV_X86_MM3      =  149, PDB_CV_X86_XMM3     =  157,
+  PDB_CV_X86_MM4      =  150, PDB_CV_X86_XMM4     =  158,
+  PDB_CV_X86_MM5      =  151, PDB_CV_X86_XMM5     =  159,
+  PDB_CV_X86_MM6      =  152, PDB_CV_X86_XMM6     =  160,
+  PDB_CV_X86_MM7      =  153, PDB_CV_X86_XMM7     =  161,
+
+  PDB_CV_X86_XMM00    =  162,  PDB_CV_X86_XMM33    =  177,
+  PDB_CV_X86_XMM01    =  163,  PDB_CV_X86_XMM40    =  178,
+  PDB_CV_X86_XMM02    =  164,  PDB_CV_X86_XMM41    =  179,
+  PDB_CV_X86_XMM03    =  165,  PDB_CV_X86_XMM42    =  180,
+  PDB_CV_X86_XMM10    =  166,  PDB_CV_X86_XMM43    =  181,
+  PDB_CV_X86_XMM11    =  167,  PDB_CV_X86_XMM50    =  182,
+  PDB_CV_X86_XMM12    =  168,  PDB_CV_X86_XMM51    =  183,
+  PDB_CV_X86_XMM13    =  169,  PDB_CV_X86_XMM52    =  184,
+  PDB_CV_X86_XMM20    =  170,  PDB_CV_X86_XMM53    =  185,
+  PDB_CV_X86_XMM21    =  171,  PDB_CV_X86_XMM60    =  186,
+  PDB_CV_X86_XMM22    =  172,  PDB_CV_X86_XMM61    =  187,
+  PDB_CV_X86_XMM23    =  173,  PDB_CV_X86_XMM62    =  188,
+  PDB_CV_X86_XMM30    =  174,  PDB_CV_X86_XMM63    =  189,
+  PDB_CV_X86_XMM31    =  175,  PDB_CV_X86_XMM70    =  190,
+  PDB_CV_X86_XMM32    =  176,  PDB_CV_X86_XMM71    =  191,
+  PDB_CV_X86_XMM73    =  193,  PDB_CV_X86_XMM72    =  192,
+
+  PDB_CV_X86_XMM0L    =  194, PDB_CV_X86_XMM0H    =  202,
+  PDB_CV_X86_XMM1L    =  195, PDB_CV_X86_XMM1H    =  203,
+  PDB_CV_X86_XMM2L    =  196, PDB_CV_X86_XMM2H    =  204,
+  PDB_CV_X86_XMM3L    =  197, PDB_CV_X86_XMM3H    =  205,
+  PDB_CV_X86_XMM4L    =  198, PDB_CV_X86_XMM4H    =  206,
+  PDB_CV_X86_XMM5L    =  199, PDB_CV_X86_XMM5H    =  207,
+  PDB_CV_X86_XMM6L    =  200, PDB_CV_X86_XMM6H    =  208,
+  PDB_CV_X86_XMM7L    =  201, PDB_CV_X86_XMM7H    =  209,
+
+  /* AVX registers */
+
+  PDB_CV_X86_YMM0 = 252, PDB_CV_X86_YMM0H = 260,
+  PDB_CV_X86_YMM1 = 253, PDB_CV_X86_YMM1H = 261,
+  PDB_CV_X86_YMM2 = 254, PDB_CV_X86_YMM2H = 262,
+  PDB_CV_X86_YMM3 = 255, PDB_CV_X86_YMM3H = 263,
+  PDB_CV_X86_YMM4 = 256, PDB_CV_X86_YMM4H = 264,
+  PDB_CV_X86_YMM5 = 257, PDB_CV_X86_YMM5H = 265,
+  PDB_CV_X86_YMM6 = 258, PDB_CV_X86_YMM6H = 266,
+  PDB_CV_X86_YMM7 = 259, PDB_CV_X86_YMM7H = 267,
+
+  /* AVX integer registers */
+
+  PDB_CV_X86_YMM0I0 = 268, PDB_CV_X86_YMM4I0 = 284, 
+  PDB_CV_X86_YMM0I1 = 269, PDB_CV_X86_YMM4I1 = 285, 
+  PDB_CV_X86_YMM0I2 = 270, PDB_CV_X86_YMM4I2 = 286,
+  PDB_CV_X86_YMM0I3 = 271, PDB_CV_X86_YMM4I3 = 287,
+  PDB_CV_X86_YMM1I0 = 272, PDB_CV_X86_YMM5I0 = 288,
+  PDB_CV_X86_YMM1I1 = 273, PDB_CV_X86_YMM5I1 = 289,
+  PDB_CV_X86_YMM1I2 = 274, PDB_CV_X86_YMM5I2 = 290,
+  PDB_CV_X86_YMM1I3 = 275, PDB_CV_X86_YMM5I3 = 291,
+  PDB_CV_X86_YMM2I0 = 276, PDB_CV_X86_YMM6I0 = 292,
+  PDB_CV_X86_YMM2I1 = 277, PDB_CV_X86_YMM6I1 = 293,
+  PDB_CV_X86_YMM2I2 = 278, PDB_CV_X86_YMM6I2 = 294,
+  PDB_CV_X86_YMM2I3 = 279, PDB_CV_X86_YMM6I3 = 295,
+  PDB_CV_X86_YMM3I0 = 280, PDB_CV_X86_YMM7I0 = 296,
+  PDB_CV_X86_YMM3I1 = 281, PDB_CV_X86_YMM7I1 = 297,
+  PDB_CV_X86_YMM3I2 = 282, PDB_CV_X86_YMM7I2 = 298,
+  PDB_CV_X86_YMM3I3 = 283, PDB_CV_X86_YMM7I3 = 299,
+
+  /* AVX floating-point single precise registers */
+
+  PDB_CV_X86_YMM0F0 = 300, PDB_CV_X86_YMM4F0 = 332,
+  PDB_CV_X86_YMM0F1 = 301, PDB_CV_X86_YMM4F1 = 333,
+  PDB_CV_X86_YMM0F2 = 302, PDB_CV_X86_YMM4F2 = 334,
+  PDB_CV_X86_YMM0F3 = 303, PDB_CV_X86_YMM4F3 = 335,
+  PDB_CV_X86_YMM0F4 = 304, PDB_CV_X86_YMM4F4 = 336,
+  PDB_CV_X86_YMM0F5 = 305, PDB_CV_X86_YMM4F5 = 337,
+  PDB_CV_X86_YMM0F6 = 306, PDB_CV_X86_YMM4F6 = 338,
+  PDB_CV_X86_YMM0F7 = 307, PDB_CV_X86_YMM4F7 = 339,
+  PDB_CV_X86_YMM1F0 = 308, PDB_CV_X86_YMM5F0 = 340,
+  PDB_CV_X86_YMM1F1 = 309, PDB_CV_X86_YMM5F1 = 341,
+  PDB_CV_X86_YMM1F2 = 310, PDB_CV_X86_YMM5F2 = 342,
+  PDB_CV_X86_YMM1F3 = 311, PDB_CV_X86_YMM5F3 = 343,
+  PDB_CV_X86_YMM1F4 = 312, PDB_CV_X86_YMM5F4 = 344,
+  PDB_CV_X86_YMM1F5 = 313, PDB_CV_X86_YMM5F5 = 345,
+  PDB_CV_X86_YMM1F6 = 314, PDB_CV_X86_YMM5F6 = 346,
+  PDB_CV_X86_YMM1F7 = 315, PDB_CV_X86_YMM5F7 = 347,
+  PDB_CV_X86_YMM2F0 = 316, PDB_CV_X86_YMM6F0 = 348,
+  PDB_CV_X86_YMM2F1 = 317, PDB_CV_X86_YMM6F1 = 349,
+  PDB_CV_X86_YMM2F2 = 318, PDB_CV_X86_YMM6F2 = 350,
+  PDB_CV_X86_YMM2F3 = 319, PDB_CV_X86_YMM6F3 = 351,
+  PDB_CV_X86_YMM2F4 = 320, PDB_CV_X86_YMM6F4 = 352,
+  PDB_CV_X86_YMM2F5 = 321, PDB_CV_X86_YMM6F5 = 353,
+  PDB_CV_X86_YMM2F6 = 322, PDB_CV_X86_YMM6F6 = 354,
+  PDB_CV_X86_YMM2F7 = 323, PDB_CV_X86_YMM6F7 = 355,
+  PDB_CV_X86_YMM3F0 = 324, PDB_CV_X86_YMM7F0 = 356,
+  PDB_CV_X86_YMM3F1 = 325, PDB_CV_X86_YMM7F1 = 357,
+  PDB_CV_X86_YMM3F2 = 326, PDB_CV_X86_YMM7F2 = 358,
+  PDB_CV_X86_YMM3F3 = 327, PDB_CV_X86_YMM7F3 = 359,
+  PDB_CV_X86_YMM3F4 = 328, PDB_CV_X86_YMM7F4 = 360,
+  PDB_CV_X86_YMM3F5 = 329, PDB_CV_X86_YMM7F5 = 361,
+  PDB_CV_X86_YMM3F6 = 330, PDB_CV_X86_YMM7F6 = 362,
+  PDB_CV_X86_YMM3F7 = 331, PDB_CV_X86_YMM7F7 = 363,
+
+  /* AVX floating-point double precise registers */
+
+  PDB_CV_X86_YMM0D0 = 364, PDB_CV_X86_YMM4D0 = 380,
+  PDB_CV_X86_YMM0D1 = 365, PDB_CV_X86_YMM4D1 = 381,
+  PDB_CV_X86_YMM0D2 = 366, PDB_CV_X86_YMM4D2 = 382,
+  PDB_CV_X86_YMM0D3 = 367, PDB_CV_X86_YMM4D3 = 383,
+  PDB_CV_X86_YMM1D0 = 368, PDB_CV_X86_YMM5D0 = 384,
+  PDB_CV_X86_YMM1D1 = 369, PDB_CV_X86_YMM5D1 = 385,
+  PDB_CV_X86_YMM1D2 = 370, PDB_CV_X86_YMM5D2 = 386,
+  PDB_CV_X86_YMM1D3 = 371, PDB_CV_X86_YMM5D3 = 387,
+  PDB_CV_X86_YMM2D0 = 372, PDB_CV_X86_YMM6D0 = 388,
+  PDB_CV_X86_YMM2D1 = 373, PDB_CV_X86_YMM6D1 = 389,
+  PDB_CV_X86_YMM2D2 = 374, PDB_CV_X86_YMM6D2 = 390,
+  PDB_CV_X86_YMM2D3 = 375, PDB_CV_X86_YMM6D3 = 391,
+  PDB_CV_X86_YMM3D0 = 376, PDB_CV_X86_YMM7D0 = 392,
+  PDB_CV_X86_YMM3D1 = 377, PDB_CV_X86_YMM7D1 = 393,
+  PDB_CV_X86_YMM3D2 = 378, PDB_CV_X86_YMM7D2 = 394,
+  PDB_CV_X86_YMM3D3 = 379, PDB_CV_X86_YMM7D3 = 395,
+
+  /* AMD64 registers */
+
+  PDB_CV_X64_AL  =  1, PDB_CV_X64_ECX    = 18,
+  PDB_CV_X64_CL  =  2, PDB_CV_X64_EDX    = 19,
+  PDB_CV_X64_DL  =  3, PDB_CV_X64_EBX    = 20,
+  PDB_CV_X64_BL  =  4, PDB_CV_X64_ESP    = 21,
+  PDB_CV_X64_AH  =  5, PDB_CV_X64_EBP    = 22,
+  PDB_CV_X64_CH  =  6, PDB_CV_X64_ESI    = 23,
+  PDB_CV_X64_DH  =  7, PDB_CV_X64_EDI    = 24,
+  PDB_CV_X64_BH  =  8, PDB_CV_X64_ES     = 25,
+  PDB_CV_X64_AX  =  9, PDB_CV_X64_CS     = 26,
+  PDB_CV_X64_CX  = 10, PDB_CV_X64_SS     = 27,
+  PDB_CV_X64_DX  = 11, PDB_CV_X64_DS     = 28,
+  PDB_CV_X64_BX  = 12, PDB_CV_X64_FS     = 29,
+  PDB_CV_X64_SP  = 13, PDB_CV_X64_GS     = 30,
+  PDB_CV_X64_BP  = 14, PDB_CV_X64_FLAGS  = 32,
+  PDB_CV_X64_SI  = 15, PDB_CV_X64_RIP    = 33,
+  PDB_CV_X64_DI  = 16, PDB_CV_X64_EFLAGS = 34,
+  PDB_CV_X64_EAX = 17,
+
+  /* Control registers */
+
+  PDB_CV_X64_CR0 = 80,
+  PDB_CV_X64_CR1 = 81,
+  PDB_CV_X64_CR2 = 82,
+  PDB_CV_X64_CR3 = 83,
+  PDB_CV_X64_CR4 = 84,
+  PDB_CV_X64_CR8 = 88,
+
+  /* Debug registers */
+
+  PDB_CV_X64_DR0 = 90, PDB_CV_X64_DR8  = 98,
+  PDB_CV_X64_DR1 = 91, PDB_CV_X64_DR9  = 99,
+  PDB_CV_X64_DR2 = 92, PDB_CV_X64_DR10 = 100,
+  PDB_CV_X64_DR3 = 93, PDB_CV_X64_DR11 = 101,
+  PDB_CV_X64_DR4 = 94, PDB_CV_X64_DR12 = 102,
+  PDB_CV_X64_DR5 = 95, PDB_CV_X64_DR13 = 103,
+  PDB_CV_X64_DR6 = 96, PDB_CV_X64_DR14 = 104,
+  PDB_CV_X64_DR7 = 97, PDB_CV_X64_DR15 = 105,
+
+  PDB_CV_X64_GDTR = 110, PDB_CV_X64_IDTL = 113,
+  PDB_CV_X64_GDTL = 111, PDB_CV_X64_LDTR = 114,
+  PDB_CV_X64_IDTR = 112, PDB_CV_X64_TR   = 115,
+
+  PDB_CV_X64_ST0  = 128, PDB_CV_X64_STAT  = 137,
+  PDB_CV_X64_ST1  = 129, PDB_CV_X64_TAG   = 138,
+  PDB_CV_X64_ST2  = 130, PDB_CV_X64_FPIP  = 139,
+  PDB_CV_X64_ST3  = 131, PDB_CV_X64_FPCS  = 140,
+  PDB_CV_X64_ST4  = 132, PDB_CV_X64_FPDO  = 141,
+  PDB_CV_X64_ST5  = 133, PDB_CV_X64_FPDS  = 142,
+  PDB_CV_X64_ST6  = 134, PDB_CV_X64_ISEM  = 143,
+  PDB_CV_X64_ST7  = 135, PDB_CV_X64_FPEIP = 144,
+  PDB_CV_X64_CTRL = 136, PDB_CV_X64_FPEDO = 145,
+
+  PDB_CV_X64_MM0 = 146, PDB_CV_X64_MM4 = 150,
+  PDB_CV_X64_MM1 = 147, PDB_CV_X64_MM5 = 151,
+  PDB_CV_X64_MM2 = 148, PDB_CV_X64_MM6 = 152,
+  PDB_CV_X64_MM3 = 149, PDB_CV_X64_MM7 = 153,
+
+  PDB_CV_X64_XMM0 = 154, PDB_CV_X64_XMM4 = 158,
+  PDB_CV_X64_XMM1 = 155, PDB_CV_X64_XMM5 = 159,
+  PDB_CV_X64_XMM2 = 156, PDB_CV_X64_XMM6 = 160,
+  PDB_CV_X64_XMM3 = 157, PDB_CV_X64_XMM7 = 161,
+
+  PDB_CV_X64_XMM0_0 = 162, PDB_CV_X64_XMM2_0 = 170, PDB_CV_X64_XMM4_0 = 178, PDB_CV_X64_XMM6_0 = 186,
+  PDB_CV_X64_XMM0_1 = 163, PDB_CV_X64_XMM2_1 = 171, PDB_CV_X64_XMM4_1 = 179, PDB_CV_X64_XMM6_1 = 187,
+  PDB_CV_X64_XMM0_2 = 164, PDB_CV_X64_XMM2_2 = 172, PDB_CV_X64_XMM4_2 = 180, PDB_CV_X64_XMM6_2 = 188,
+  PDB_CV_X64_XMM0_3 = 165, PDB_CV_X64_XMM2_3 = 173, PDB_CV_X64_XMM4_3 = 181, PDB_CV_X64_XMM6_3 = 189,
+  PDB_CV_X64_XMM1_0 = 166, PDB_CV_X64_XMM3_0 = 174, PDB_CV_X64_XMM5_0 = 182, PDB_CV_X64_XMM7_0 = 190,
+  PDB_CV_X64_XMM1_1 = 167, PDB_CV_X64_XMM3_1 = 175, PDB_CV_X64_XMM5_1 = 183, PDB_CV_X64_XMM7_1 = 191,
+  PDB_CV_X64_XMM1_2 = 168, PDB_CV_X64_XMM3_2 = 176, PDB_CV_X64_XMM5_2 = 184, PDB_CV_X64_XMM7_2 = 192,
+  PDB_CV_X64_XMM1_3 = 169, PDB_CV_X64_XMM3_3 = 177, PDB_CV_X64_XMM5_3 = 185, PDB_CV_X64_XMM7_3 = 193,
+
+  PDB_CV_X64_XMM0L = 194, PDB_CV_X64_XMM4L = 198,
+  PDB_CV_X64_XMM1L = 195, PDB_CV_X64_XMM5L = 199,
+  PDB_CV_X64_XMM2L = 196, PDB_CV_X64_XMM6L = 200,
+  PDB_CV_X64_XMM3L = 197, PDB_CV_X64_XMM7L = 201,
+
+  PDB_CV_X64_XMM0H = 202, PDB_CV_X64_XMM4H = 206,
+  PDB_CV_X64_XMM1H = 203, PDB_CV_X64_XMM5H = 207,
+  PDB_CV_X64_XMM2H = 204, PDB_CV_X64_XMM6H = 208,
+  PDB_CV_X64_XMM3H = 205, PDB_CV_X64_XMM7H = 209,
+
+  /* XMM status register */
+  PDB_CV_X64_MXCSR = 211,    
+
+  /* XMM sub-registers (WNI integer) */
+  PDB_CV_X64_EMM0L = 220, PDB_CV_X64_EMM4L = 224,
+  PDB_CV_X64_EMM1L = 221, PDB_CV_X64_EMM5L = 225,
+  PDB_CV_X64_EMM2L = 222, PDB_CV_X64_EMM6L = 226,
+  PDB_CV_X64_EMM3L = 223, PDB_CV_X64_EMM7L = 227,
+
+  PDB_CV_X64_EMM0H = 228, PDB_CV_X64_EMM4H = 232,
+  PDB_CV_X64_EMM1H = 229, PDB_CV_X64_EMM5H = 233,
+  PDB_CV_X64_EMM2H = 230, PDB_CV_X64_EMM6H = 234,
+  PDB_CV_X64_EMM3H = 231, PDB_CV_X64_EMM7H = 235,
+  
+  /* do not change the order of these regs, first one must be even too */
+
+  PDB_CV_X64_MM00 = 236, PDB_CV_X64_MM40 = 244,
+  PDB_CV_X64_MM01 = 237, PDB_CV_X64_MM41 = 245,
+  PDB_CV_X64_MM10 = 238, PDB_CV_X64_MM50 = 246,
+  PDB_CV_X64_MM11 = 239, PDB_CV_X64_MM51 = 247,
+  PDB_CV_X64_MM20 = 240, PDB_CV_X64_MM60 = 248,
+  PDB_CV_X64_MM21 = 241, PDB_CV_X64_MM61 = 249,
+  PDB_CV_X64_MM30 = 242, PDB_CV_X64_MM70 = 250,
+  PDB_CV_X64_MM31 = 243, PDB_CV_X64_MM71 = 251,
+  
+  PDB_CV_X64_XMM8  = 252, PDB_CV_X64_XMM12 = 256,
+  PDB_CV_X64_XMM9  = 253, PDB_CV_X64_XMM13 = 257,
+  PDB_CV_X64_XMM10 = 254, PDB_CV_X64_XMM14 = 258,
+  PDB_CV_X64_XMM11 = 255, PDB_CV_X64_XMM15 = 259,
+
+  PDB_CV_X64_XMM8_0  = 260, PDB_CV_X64_XMM12_0 = 276,
+  PDB_CV_X64_XMM8_1  = 261, PDB_CV_X64_XMM12_1 = 277,
+  PDB_CV_X64_XMM8_2  = 262, PDB_CV_X64_XMM12_2 = 278,
+  PDB_CV_X64_XMM8_3  = 263, PDB_CV_X64_XMM12_3 = 279,
+  PDB_CV_X64_XMM9_0  = 264, PDB_CV_X64_XMM13_0 = 280,
+  PDB_CV_X64_XMM9_1  = 265, PDB_CV_X64_XMM13_1 = 281,
+  PDB_CV_X64_XMM9_2  = 266, PDB_CV_X64_XMM13_2 = 282,
+  PDB_CV_X64_XMM9_3  = 267, PDB_CV_X64_XMM13_3 = 283,
+  PDB_CV_X64_XMM10_0 = 268, PDB_CV_X64_XMM14_0 = 284,
+  PDB_CV_X64_XMM10_1 = 269, PDB_CV_X64_XMM14_1 = 285,
+  PDB_CV_X64_XMM10_2 = 270, PDB_CV_X64_XMM14_2 = 286,
+  PDB_CV_X64_XMM10_3 = 271, PDB_CV_X64_XMM14_3 = 287,
+  PDB_CV_X64_XMM11_0 = 272, PDB_CV_X64_XMM15_0 = 288,
+  PDB_CV_X64_XMM11_1 = 273, PDB_CV_X64_XMM15_1 = 289,
+  PDB_CV_X64_XMM11_2 = 274, PDB_CV_X64_XMM15_2 = 290,
+  PDB_CV_X64_XMM11_3 = 275, PDB_CV_X64_XMM15_3 = 291,
+
+  PDB_CV_X64_XMM8L  = 292, PDB_CV_X64_XMM12L = 296,
+  PDB_CV_X64_XMM9L  = 293, PDB_CV_X64_XMM13L = 297,
+  PDB_CV_X64_XMM10L = 294, PDB_CV_X64_XMM14L = 298,
+  PDB_CV_X64_XMM11L = 295, PDB_CV_X64_XMM15L = 299,
+
+  PDB_CV_X64_XMM8H  = 300, PDB_CV_X64_XMM12H = 304,
+  PDB_CV_X64_XMM9H  = 301, PDB_CV_X64_XMM13H = 305,
+  PDB_CV_X64_XMM10H = 302, PDB_CV_X64_XMM14H = 306,
+  PDB_CV_X64_XMM11H = 303, PDB_CV_X64_XMM15H = 307,
+
+  /* XMM sub-registers (WNI integer) */
+
+  PDB_CV_X64_EMM8L  = 308, PDB_CV_X64_EMM12L = 312,
+  PDB_CV_X64_EMM9L  = 309, PDB_CV_X64_EMM13L = 313,
+  PDB_CV_X64_EMM10L = 310, PDB_CV_X64_EMM14L = 314,
+  PDB_CV_X64_EMM11L = 311, PDB_CV_X64_EMM15L = 315,
+
+  PDB_CV_X64_EMM8H  = 316, PDB_CV_X64_EMM12H = 320,
+  PDB_CV_X64_EMM9H  = 317, PDB_CV_X64_EMM13H = 321,
+  PDB_CV_X64_EMM10H = 318, PDB_CV_X64_EMM14H = 322,
+  PDB_CV_X64_EMM11H = 319, PDB_CV_X64_EMM15H = 323,
+
+  /* Low byte forms of some standard registers */
+  PDB_CV_X64_SIL = 324, PDB_CV_X64_BPL = 326,
+  PDB_CV_X64_DIL = 325, PDB_CV_X64_SPL = 327,
+  
+  /* 64-bit regular registers */
+  PDB_CV_X64_RAX = 328, PDB_CV_X64_RSP = 335,
+  PDB_CV_X64_RSI = 332, PDB_CV_X64_RBX = 329,
+  PDB_CV_X64_RDI = 333, PDB_CV_X64_RCX = 330,
+  PDB_CV_X64_RBP = 334, PDB_CV_X64_RDX = 331,
+
+  /* 64-bit integer registers with 8-, 16-, and 32-bit forms (B, W, and D) */
+  PDB_CV_X64_R8  = 336, PDB_CV_X64_R12 = 340,
+  PDB_CV_X64_R9  = 337, PDB_CV_X64_R13 = 341,
+  PDB_CV_X64_R10 = 338, PDB_CV_X64_R14 = 342,
+  PDB_CV_X64_R11 = 339, PDB_CV_X64_R15 = 343,
+
+  PDB_CV_X64_R8B  = 344, PDB_CV_X64_R12B = 348,
+  PDB_CV_X64_R9B  = 345, PDB_CV_X64_R13B = 349,
+  PDB_CV_X64_R10B = 346, PDB_CV_X64_R14B = 350,
+  PDB_CV_X64_R11B = 347, PDB_CV_X64_R15B = 351,
+
+  PDB_CV_X64_R8W  = 352, PDB_CV_X64_R12W = 356,
+  PDB_CV_X64_R9W  = 353, PDB_CV_X64_R13W = 357,
+  PDB_CV_X64_R10W = 354, PDB_CV_X64_R14W = 358,
+  PDB_CV_X64_R11W = 355, PDB_CV_X64_R15W = 359,
+
+  PDB_CV_X64_R8D  = 360, PDB_CV_X64_R12D = 364,
+  PDB_CV_X64_R9D  = 361, PDB_CV_X64_R13D = 365,
+  PDB_CV_X64_R10D = 362, PDB_CV_X64_R14D = 366,
+  PDB_CV_X64_R11D = 363, PDB_CV_X64_R15D = 367,
+
+  /* AVX registers 256 bits */
+
+  PDB_CV_X64_YMM0 = 368, PDB_CV_X64_YMM8  = 376,
+  PDB_CV_X64_YMM1 = 369, PDB_CV_X64_YMM9  = 377,
+  PDB_CV_X64_YMM2 = 370, PDB_CV_X64_YMM10 = 378,
+  PDB_CV_X64_YMM3 = 371, PDB_CV_X64_YMM11 = 379,
+  PDB_CV_X64_YMM4 = 372, PDB_CV_X64_YMM12 = 380,
+  PDB_CV_X64_YMM5 = 373, PDB_CV_X64_YMM13 = 381,
+  PDB_CV_X64_YMM6 = 374, PDB_CV_X64_YMM14 = 382,
+  PDB_CV_X64_YMM7 = 375, PDB_CV_X64_YMM15 = 383,
+
+  /* AVX registers upper 128 bits */
+  PDB_CV_X64_YMM0H = 384, PDB_CV_X64_YMM8H  = 392,
+  PDB_CV_X64_YMM1H = 385, PDB_CV_X64_YMM9H  = 393,
+  PDB_CV_X64_YMM2H = 386, PDB_CV_X64_YMM10H = 394,
+  PDB_CV_X64_YMM3H = 387, PDB_CV_X64_YMM11H = 395,
+  PDB_CV_X64_YMM4H = 388, PDB_CV_X64_YMM12H = 396,
+  PDB_CV_X64_YMM5H = 389, PDB_CV_X64_YMM13H = 397,
+  PDB_CV_X64_YMM6H = 390, PDB_CV_X64_YMM14H = 398,
+  PDB_CV_X64_YMM7H = 391, PDB_CV_X64_YMM15H = 399,
+
+  /* Lower/upper 8 bytes of XMM registers.  Unlike CV_AMD64_XMM<regnum><H/L>, these
+   * values reprsesent the bit patterns of the registers as 64-bit integers, not
+   * the representation of these registers as a double. */
+  PDB_CV_X64_XMM0IL = 400, PDB_CV_X64_XMM8IL  = 408,
+  PDB_CV_X64_XMM1IL = 401, PDB_CV_X64_XMM9IL  = 409,
+  PDB_CV_X64_XMM2IL = 402, PDB_CV_X64_XMM10IL = 410,
+  PDB_CV_X64_XMM3IL = 403, PDB_CV_X64_XMM11IL = 411,
+  PDB_CV_X64_XMM4IL = 404, PDB_CV_X64_XMM12IL = 412,
+  PDB_CV_X64_XMM5IL = 405, PDB_CV_X64_XMM13IL = 413,
+  PDB_CV_X64_XMM6IL = 406, PDB_CV_X64_XMM14IL = 414,
+  PDB_CV_X64_XMM7IL = 407, PDB_CV_X64_XMM15IL = 415,
+
+  PDB_CV_X64_XMM0IH = 416, PDB_CV_X64_XMM8IH  = 424,
+  PDB_CV_X64_XMM1IH = 417, PDB_CV_X64_XMM9IH  = 425,
+  PDB_CV_X64_XMM2IH = 418, PDB_CV_X64_XMM10IH = 426,
+  PDB_CV_X64_XMM3IH = 419, PDB_CV_X64_XMM11IH = 427,
+  PDB_CV_X64_XMM4IH = 420, PDB_CV_X64_XMM12IH = 428,
+  PDB_CV_X64_XMM5IH = 421, PDB_CV_X64_XMM13IH = 429,
+  PDB_CV_X64_XMM6IH = 422, PDB_CV_X64_XMM14IH = 430,
+  PDB_CV_X64_XMM7IH = 423, PDB_CV_X64_XMM15IH = 431,
+
+  /* AVX integer registers */
+
+  PDB_CV_X64_YMM0I0 = 432, PDB_CV_X64_YMM4I0 = 448, PDB_CV_X64_YMM8I0  = 464, PDB_CV_X64_YMM12I0 = 480,
+  PDB_CV_X64_YMM0I1 = 433, PDB_CV_X64_YMM4I1 = 449, PDB_CV_X64_YMM8I1  = 465, PDB_CV_X64_YMM12I1 = 481,
+  PDB_CV_X64_YMM0I2 = 434, PDB_CV_X64_YMM4I2 = 450, PDB_CV_X64_YMM8I2  = 466, PDB_CV_X64_YMM12I2 = 482,
+  PDB_CV_X64_YMM0I3 = 435, PDB_CV_X64_YMM4I3 = 451, PDB_CV_X64_YMM8I3  = 467, PDB_CV_X64_YMM12I3 = 483,
+  PDB_CV_X64_YMM1I0 = 436, PDB_CV_X64_YMM5I0 = 452, PDB_CV_X64_YMM9I0  = 468, PDB_CV_X64_YMM13I0 = 484,
+  PDB_CV_X64_YMM1I1 = 437, PDB_CV_X64_YMM5I1 = 453, PDB_CV_X64_YMM9I1  = 469, PDB_CV_X64_YMM13I1 = 485,
+  PDB_CV_X64_YMM1I2 = 438, PDB_CV_X64_YMM5I2 = 454, PDB_CV_X64_YMM9I2  = 470, PDB_CV_X64_YMM13I2 = 486,
+  PDB_CV_X64_YMM1I3 = 439, PDB_CV_X64_YMM5I3 = 455, PDB_CV_X64_YMM9I3  = 471, PDB_CV_X64_YMM13I3 = 487,
+  PDB_CV_X64_YMM2I0 = 440, PDB_CV_X64_YMM6I0 = 456, PDB_CV_X64_YMM10I0 = 472, PDB_CV_X64_YMM14I0 = 488,
+  PDB_CV_X64_YMM2I1 = 441, PDB_CV_X64_YMM6I1 = 457, PDB_CV_X64_YMM10I1 = 473, PDB_CV_X64_YMM14I1 = 489,
+  PDB_CV_X64_YMM2I2 = 442, PDB_CV_X64_YMM6I2 = 458, PDB_CV_X64_YMM10I2 = 474, PDB_CV_X64_YMM14I2 = 490,
+  PDB_CV_X64_YMM2I3 = 443, PDB_CV_X64_YMM6I3 = 459, PDB_CV_X64_YMM10I3 = 475, PDB_CV_X64_YMM14I3 = 491,
+  PDB_CV_X64_YMM3I0 = 444, PDB_CV_X64_YMM7I0 = 460, PDB_CV_X64_YMM11I0 = 476, PDB_CV_X64_YMM15I0 = 492,
+  PDB_CV_X64_YMM3I1 = 445, PDB_CV_X64_YMM7I1 = 461, PDB_CV_X64_YMM11I1 = 477, PDB_CV_X64_YMM15I1 = 493,
+  PDB_CV_X64_YMM3I2 = 446, PDB_CV_X64_YMM7I2 = 462, PDB_CV_X64_YMM11I2 = 478, PDB_CV_X64_YMM15I2 = 494,
+  PDB_CV_X64_YMM3I3 = 447, PDB_CV_X64_YMM7I3 = 463, PDB_CV_X64_YMM11I3 = 479, PDB_CV_X64_YMM15I3 = 495,
+
+  /* AVX floating-point single precise registers */
+
+  PDB_CV_X64_YMM0F0 = 496, PDB_CV_X64_YMM4F1 = 529, PDB_CV_X64_YMM8F2  = 562, PDB_CV_X64_YMM12F3 = 595,
+  PDB_CV_X64_YMM0F1 = 497, PDB_CV_X64_YMM4F2 = 530, PDB_CV_X64_YMM8F3  = 563, PDB_CV_X64_YMM12F4 = 596,
+  PDB_CV_X64_YMM0F2 = 498, PDB_CV_X64_YMM4F3 = 531, PDB_CV_X64_YMM8F4  = 564, PDB_CV_X64_YMM12F5 = 597,
+  PDB_CV_X64_YMM0F3 = 499, PDB_CV_X64_YMM4F4 = 532, PDB_CV_X64_YMM8F5  = 565, PDB_CV_X64_YMM12F6 = 598,
+  PDB_CV_X64_YMM0F4 = 500, PDB_CV_X64_YMM4F5 = 533, PDB_CV_X64_YMM8F6  = 566, PDB_CV_X64_YMM12F7 = 599,
+  PDB_CV_X64_YMM0F5 = 501, PDB_CV_X64_YMM4F6 = 534, PDB_CV_X64_YMM8F7  = 567, PDB_CV_X64_YMM13F0 = 600,
+  PDB_CV_X64_YMM0F6 = 502, PDB_CV_X64_YMM4F7 = 535, PDB_CV_X64_YMM9F0  = 568, PDB_CV_X64_YMM13F1 = 601,
+  PDB_CV_X64_YMM0F7 = 503, PDB_CV_X64_YMM5F0 = 536, PDB_CV_X64_YMM9F1  = 569, PDB_CV_X64_YMM13F2 = 602,
+  PDB_CV_X64_YMM1F0 = 504, PDB_CV_X64_YMM5F1 = 537, PDB_CV_X64_YMM9F2  = 570, PDB_CV_X64_YMM13F3 = 603,
+  PDB_CV_X64_YMM1F1 = 505, PDB_CV_X64_YMM5F2 = 538, PDB_CV_X64_YMM9F3  = 571, PDB_CV_X64_YMM13F4 = 604,
+  PDB_CV_X64_YMM1F2 = 506, PDB_CV_X64_YMM5F3 = 539, PDB_CV_X64_YMM9F4  = 572, PDB_CV_X64_YMM13F5 = 605,
+  PDB_CV_X64_YMM1F3 = 507, PDB_CV_X64_YMM5F4 = 540, PDB_CV_X64_YMM9F5  = 573, PDB_CV_X64_YMM13F6 = 606,
+  PDB_CV_X64_YMM1F4 = 508, PDB_CV_X64_YMM5F5 = 541, PDB_CV_X64_YMM9F6  = 574, PDB_CV_X64_YMM13F7 = 607,
+  PDB_CV_X64_YMM1F5 = 509, PDB_CV_X64_YMM5F6 = 542, PDB_CV_X64_YMM9F7  = 575, PDB_CV_X64_YMM14F0 = 608,
+  PDB_CV_X64_YMM1F6 = 510, PDB_CV_X64_YMM5F7 = 543, PDB_CV_X64_YMM10F0 = 576, PDB_CV_X64_YMM14F1 = 609,
+  PDB_CV_X64_YMM1F7 = 511, PDB_CV_X64_YMM6F0 = 544, PDB_CV_X64_YMM10F1 = 577, PDB_CV_X64_YMM14F2 = 610,
+  PDB_CV_X64_YMM2F0 = 512, PDB_CV_X64_YMM6F1 = 545, PDB_CV_X64_YMM10F2 = 578, PDB_CV_X64_YMM14F3 = 611,
+  PDB_CV_X64_YMM2F1 = 513, PDB_CV_X64_YMM6F2 = 546, PDB_CV_X64_YMM10F3 = 579, PDB_CV_X64_YMM14F4 = 612,
+  PDB_CV_X64_YMM2F2 = 514, PDB_CV_X64_YMM6F3 = 547, PDB_CV_X64_YMM10F4 = 580, PDB_CV_X64_YMM14F5 = 613,
+  PDB_CV_X64_YMM2F3 = 515, PDB_CV_X64_YMM6F4 = 548, PDB_CV_X64_YMM10F5 = 581, PDB_CV_X64_YMM14F6 = 614,
+  PDB_CV_X64_YMM2F4 = 516, PDB_CV_X64_YMM6F5 = 549, PDB_CV_X64_YMM10F6 = 582, PDB_CV_X64_YMM14F7 = 615,
+  PDB_CV_X64_YMM2F5 = 517, PDB_CV_X64_YMM6F6 = 550, PDB_CV_X64_YMM10F7 = 583, PDB_CV_X64_YMM15F0 = 616,
+  PDB_CV_X64_YMM2F6 = 518, PDB_CV_X64_YMM6F7 = 551, PDB_CV_X64_YMM11F0 = 584, PDB_CV_X64_YMM15F1 = 617,
+  PDB_CV_X64_YMM2F7 = 519, PDB_CV_X64_YMM7F0 = 552, PDB_CV_X64_YMM11F1 = 585, PDB_CV_X64_YMM15F2 = 618,
+  PDB_CV_X64_YMM3F0 = 520, PDB_CV_X64_YMM7F1 = 553, PDB_CV_X64_YMM11F2 = 586, PDB_CV_X64_YMM15F3 = 619,
+  PDB_CV_X64_YMM3F1 = 521, PDB_CV_X64_YMM7F2 = 554, PDB_CV_X64_YMM11F3 = 587, PDB_CV_X64_YMM15F4 = 620,
+  PDB_CV_X64_YMM3F2 = 522, PDB_CV_X64_YMM7F3 = 555, PDB_CV_X64_YMM11F4 = 588, PDB_CV_X64_YMM15F5 = 621,
+  PDB_CV_X64_YMM3F3 = 523, PDB_CV_X64_YMM7F4 = 556, PDB_CV_X64_YMM11F5 = 589, PDB_CV_X64_YMM15F6 = 622,
+  PDB_CV_X64_YMM3F4 = 524, PDB_CV_X64_YMM7F5 = 557, PDB_CV_X64_YMM11F6 = 590, PDB_CV_X64_YMM15F7 = 623,
+  PDB_CV_X64_YMM3F5 = 525, PDB_CV_X64_YMM7F6 = 558, PDB_CV_X64_YMM11F7 = 591,
+  PDB_CV_X64_YMM3F6 = 526, PDB_CV_X64_YMM7F7 = 559, PDB_CV_X64_YMM12F0 = 592,
+  PDB_CV_X64_YMM3F7 = 527, PDB_CV_X64_YMM8F0 = 560, PDB_CV_X64_YMM12F1 = 593,
+  PDB_CV_X64_YMM4F0 = 528, PDB_CV_X64_YMM8F1 = 561, PDB_CV_X64_YMM12F2 = 594,
+
+  /* AVX floating-point double precise registers */
+
+  PDB_CV_X64_YMM0D0 = 624,  PDB_CV_X64_YMM8D0  = 656,
+  PDB_CV_X64_YMM0D1 = 625,  PDB_CV_X64_YMM8D1  = 657,
+  PDB_CV_X64_YMM0D2 = 626,  PDB_CV_X64_YMM8D2  = 658,
+  PDB_CV_X64_YMM0D3 = 627,  PDB_CV_X64_YMM8D3  = 659,
+  PDB_CV_X64_YMM1D0 = 628,  PDB_CV_X64_YMM9D0  = 660,
+  PDB_CV_X64_YMM1D1 = 629,  PDB_CV_X64_YMM9D1  = 661,
+  PDB_CV_X64_YMM1D2 = 630,  PDB_CV_X64_YMM9D2  = 662,
+  PDB_CV_X64_YMM1D3 = 631,  PDB_CV_X64_YMM9D3  = 663,
+  PDB_CV_X64_YMM2D0 = 632,  PDB_CV_X64_YMM10D0 = 664,
+  PDB_CV_X64_YMM2D1 = 633,  PDB_CV_X64_YMM10D1 = 665,
+  PDB_CV_X64_YMM2D2 = 634,  PDB_CV_X64_YMM10D2 = 666,
+  PDB_CV_X64_YMM2D3 = 635,  PDB_CV_X64_YMM10D3 = 667,
+  PDB_CV_X64_YMM3D0 = 636,  PDB_CV_X64_YMM11D0 = 668,
+  PDB_CV_X64_YMM3D1 = 637,  PDB_CV_X64_YMM11D1 = 669,
+  PDB_CV_X64_YMM3D2 = 638,  PDB_CV_X64_YMM11D2 = 670,
+  PDB_CV_X64_YMM3D3 = 639,  PDB_CV_X64_YMM11D3 = 671,
+  PDB_CV_X64_YMM4D0 = 640,  PDB_CV_X64_YMM12D0 = 672,
+  PDB_CV_X64_YMM4D1 = 641,  PDB_CV_X64_YMM12D1 = 673,
+  PDB_CV_X64_YMM4D2 = 642,  PDB_CV_X64_YMM12D2 = 674,
+  PDB_CV_X64_YMM4D3 = 643,  PDB_CV_X64_YMM12D3 = 675,
+  PDB_CV_X64_YMM5D0 = 644,  PDB_CV_X64_YMM13D0 = 676,
+  PDB_CV_X64_YMM5D1 = 645,  PDB_CV_X64_YMM13D1 = 677,
+  PDB_CV_X64_YMM5D2 = 646,  PDB_CV_X64_YMM13D2 = 678,
+  PDB_CV_X64_YMM5D3 = 647,  PDB_CV_X64_YMM13D3 = 679,
+  PDB_CV_X64_YMM6D0 = 648,  PDB_CV_X64_YMM14D0 = 680,
+  PDB_CV_X64_YMM6D1 = 649,  PDB_CV_X64_YMM14D1 = 681,
+  PDB_CV_X64_YMM6D2 = 650,  PDB_CV_X64_YMM14D2 = 682,
+  PDB_CV_X64_YMM6D3 = 651,  PDB_CV_X64_YMM14D3 = 683,
+  PDB_CV_X64_YMM7D0 = 652,  PDB_CV_X64_YMM15D0 = 684,
+  PDB_CV_X64_YMM7D1 = 653,  PDB_CV_X64_YMM15D1 = 685,
+  PDB_CV_X64_YMM7D2 = 654,  PDB_CV_X64_YMM15D2 = 686,
+  PDB_CV_X64_YMM7D3 = 655,  PDB_CV_X64_YMM15D3 = 687,
+
+  PDB_CV_REG_MAX
+} pdb_cv_reg_e;
+
+typedef enum {
+   PDB_CV_HFA_NONE   =  0,
+   PDB_CV_HFA_FLOAT  =  1,
+   PDB_CV_HFA_DOUBLE =  2,
+   PDB_CV_HFA_OTHER  =  3
+} pdb_cv_hfa_e;
+
+typedef enum {
+  PDB_CV_MOCOM_UDT_NONE      = 0,
+  PDB_CV_MOCOM_UDT_REF       = 1,
+  PDB_CV_MOCOM_UDT_VALUE     = 2,
+  PDB_CV_MOCOM_UDT_INTERFACE = 3
+} pdb_cv_mocom_udt_e;
+
+#define PDB_CV_PROP_PACKED          (1 << 0)    /* syms_true if structure is packed */
+#define PDB_CV_PROP_CTOR            (1 << 1)    /* syms_true if constructors or destructors present */
+#define PDB_CV_PROP_OVLOPS          (1 << 2)    /* syms_true if overloaded operators present  */
+#define PDB_CV_PROP_ISNSESTED       (1 << 3)    /* syms_true if this is a nested class */
+#define PDB_CV_PROP_CNESTED         (1 << 4)    /* syms_true if this class contains nested types */
+#define PDB_CV_PROP_OPASSIGN        (1 << 5)    /* syms_true if overloaded assignment (=) */
+#define PDB_CV_PROP_OPCAST          (1 << 6)    /* syms_true if casting methods */
+#define PDB_CV_PROP_FWDREF          (1 << 7)    /* syms_true if forward reference (incomplete defn) */
+#define PDB_CV_PROP_SCOPED          (1 << 8)    /* scoped definition */
+#define PDB_CV_PROP_HAS_UNIQUE_NAME (1 << 9)    /* syms_true if there is a decorated name following the regular name */
+#define PDB_CV_PROP_SEALED          (1 << 10    /* syms_true if class cannot be used as a base class */
+#define PDB_CV_PROP_INTRINSIC       (1 << 13)
+/* Mask for the pdb_cv_hfa_e */
+#define PDB_CV_PROP_HFA_MASK    0x1800
+/* Mask for the pdb_cv_mocom_udt_e */
+#define PDB_CV_PROP_MOCOM_MASK  0xC000
+typedef U16 pdb_cv_prop_t;
+
+/* Method attributes */
+#define PDB_CV_FLDATTR_MPROP_VANILLA   0x0u
+#define PDB_CV_FLDATTR_MPROP_VIRTUAL   0x1u
+#define PDB_CV_FLDATTR_MPROP_STATIC    0x2u
+#define PDB_CV_FLDATTR_MPROP_FRIEND    0x3u
+#define PDB_CV_FLDATTR_MPROP_INTRO     0x4u
+#define PDB_CV_FLDATTR_MPROP_PUREVIRT  0x5u
+#define PDB_CV_FLDATTR_MPROP_PUREINTRO 0x6u
+
+#define PDB_CV_FLDATTR_ACCESS_PRIVATE   0x1u
+#define PDB_CV_FLDATTR_ACCESS_PROTECTED 0x2u
+#define PDB_CV_FLDATTR_ACCESS_PUBLIC    0x3u
+
+#define PDB_CV_FLDATTR_ACCESS_MASK(x)   (((x) & 0x03u) >> 0u) /* access protection pdb_cv_access_t  */
+#define PDB_CV_FLDATTR_MPROP_MASK(x)    (((x) & 0x1Cu) >> 2u) /* method properties pdb_cv_methodprop_t  */
+#define PDB_CV_FLDATTR_PSEUDO      (1u << 6u)                 /* compiler generated fcn and does not exist */
+#define PDB_CV_FLDATTR_NOINHERIT   (1u << 7u)                 /* syms_true if class cannot be inherited */
+#define PDB_CV_FLDATTR_NOCONSTRUCT (1u << 8u)                 /* syms_true if class cannot be constructed */
+#define PDB_CV_FLDATTR_COMPGENX    (1u << 9u)                 /* compiler generated fcn and does exist */
+#define PDB_CV_FLDATTR_SEALED      (1u << 10u)                /* syms_true if method cannot be overridden */
+typedef U16 pdb_cv_fldattr_t;
+
+#define PDB_CV_MODIFIER_CONST       (1u << 0u)
+#define PDB_CV_MODIFIER_VOLATILE    (1u << 1u)
+#define PDB_CV_MODIFIER_UNALIGNED   (1u << 2u)
+typedef U16 pdb_cv_modifier_t;
+
+typedef enum {
+  PDB_CV_PTR_MODE_PTR     = 0x00, /* "normal" pointer */
+  PDB_CV_PTR_MODE_REF     = 0x01, /* "old" reference */
+  PDB_CV_PTR_MODE_LVREF   = 0x01, /* l-value reference */
+  PDB_CV_PTR_MODE_PMEM    = 0x02, /* pointer to data member */
+  PDB_CV_PTR_MODE_PMFUNC  = 0x03, /* pointer to member function */
+  PDB_CV_PTR_MODE_RVREF   = 0x04, /* r-value reference */
+  PDB_CV_PTR_MODE_RESERVED= 0x05  /* first unused pointer mode */
+} pdb_cv_ptrmode_e;
+typedef u32 pdb_cv_ptrmode;
+
+typedef enum {
+  PDB_CV_PTR_NEAR         = 0x00, /* 16 bit pointer */
+  PDB_CV_PTR_FAR          = 0x01, /* 16:16 far pointer */
+  PDB_CV_PTR_HUGE         = 0x02, /* 16:16 huge pointer */
+  PDB_CV_PTR_BASE_SEG     = 0x03, /* based on segment */
+  PDB_CV_PTR_BASE_VAL     = 0x04, /* based on value of base */
+  PDB_CV_PTR_BASE_SEGVAL  = 0x05, /* based on segment value of base */
+  PDB_CV_PTR_BASE_ADDR    = 0x06, /* based on address of base */
+  PDB_CV_PTR_BASE_SEGADDR = 0x07, /* based on segment address of base */
+  PDB_CV_PTR_BASE_TYPE    = 0x08, /* based on type */
+  PDB_CV_PTR_BASE_SELF    = 0x09, /* based on self */
+  PDB_CV_PTR_NEAR32       = 0x0a, /* 32 bit pointer */
+  PDB_CV_PTR_FAR32        = 0x0b, /* 16:32 pointer */
+  PDB_CV_PTR_64           = 0x0c, /* 64 bit pointer */
+  PDB_CV_PTR_UNUSEDPTR    = 0x0d  /* first unused pointer type */
+} pdb_cv_ptrtype_e;
+typedef u32 pdb_cv_ptrtype;
+
+/* ordinal specifying pointer type (pdb_cv_ptrtype_e) */
+#define PDB_CV_PTR_ATTRIB_TYPE_MASK(x) (((x) & 0x1f)    >> 0)    
+/* ordinal specifying pointer mode (pdb_cv_ptrmode_e) */
+#define PDB_CV_PTR_ATTRIB_MODE_MASK(x) (((x) & 0xe0)    >> 5) 
+/* size of pointer (in bytes) */
+#define PDB_CV_PTR_ATTRIB_SIZE_MASK(x) (((x) & 0x7E000) >> 13)
+/* syms_true if 0:32 pointer */
+#define PDB_CV_PTR_ATTRIB_IS_FLAT           (1 << 8)             
+/* TRUE if volatile pointer */
+#define PDB_CV_PTR_ATTRIB_IS_VOLATILE       (1 << 9)             
+/* TRUE if const pointer */
+#define PDB_CV_PTR_ATTRIB_IS_CONST          (1 << 10)            
+/* TRUE if unaligned pointer */
+#define PDB_CV_PTR_ATTRIB_IS_UNALIGNED      (1 << 11)            
+/* TRUE if restricted pointer (allow aggressive opts) */
+#define PDB_CV_PTR_ATTRIB_IS_RESTRICTED     (1 << 12)      
+/* TRUE if it is a MoCOM pointer (^ or %) */
+#define PDB_CV_PTR_ATTRIB_IS_MOCOM          (1 << 19)            
+/* TRUE if it is this pointer of member function with & ref-qualifier */
+#define PDB_CV_PTR_ATTRIB_IS_LREF           (1 << 20)            
+/* TRUE if it is this pointer of member function with && ref-qualifier */
+#define PDB_CV_PTR_ATTRIB_IS_RREF           (1 << 21)            
+typedef U32 pdb_cv_ptr_attrib_t;
+
+/* syms_true if C++ style ReturnUDT */
+#define PDB_CV_FUNCATTR_CXXRETURNUDT (1 << 0)   
+/* syms_true if func is an instance constructor */
+#define PDB_CV_FUNCATTR_CTOR (1 << 1)
+/* syms_true if func is an instance constructor of a class with virtual bases */
+#define PDB_CV_FUNCATTR_CTORBASE                
+typedef U8 pdb_cv_funcattr_t;
+
+typedef enum {
+  PDB_CV_CALL_NEAR_C      = 0x00, /* near right to left push, caller pops stack */
+  PDB_CV_CALL_FAR_C       = 0x01, /* far right to left push, caller pops stack */
+  PDB_CV_CALL_NEAR_PASCAL = 0x02, /* near left to right push, callee pops stack */
+  PDB_CV_CALL_FAR_PASCAL  = 0x03, /* far left to right push, callee pops stack */
+  PDB_CV_CALL_NEAR_FAST   = 0x04, /* near left to right push with regs, callee pops stack */
+  PDB_CV_CALL_FAR_FAST    = 0x05, /* far left to right push with regs, callee pops stack */
+  PDB_CV_CALL_SKIPPED     = 0x06, /* skipped (unused) call index */
+  PDB_CV_CALL_NEAR_STD    = 0x07, /* near standard call */
+  PDB_CV_CALL_FAR_STD     = 0x08, /* far standard call */
+  PDB_CV_CALL_NEAR_SYS    = 0x09, /* near sys call */
+  PDB_CV_CALL_FAR_SYS     = 0x0a, /* far sys call */
+  PDB_CV_CALL_THISCALL    = 0x0b, /* this call (this passed in register) */
+  PDB_CV_CALL_MIPSCALL    = 0x0c, /* Mips call */
+  PDB_CV_CALL_GENERIC     = 0x0d, /* Generic call sequence */
+  PDB_CV_CALL_ALPHACALL   = 0x0e, /* Alpha call */
+  PDB_CV_CALL_PPCCALL     = 0x0f, /* PPC call */
+  PDB_CV_CALL_SHCALL      = 0x10, /* Hitachi SuperH call */
+  PDB_CV_CALL_ARMCALL     = 0x11, /* ARM call */
+  PDB_CV_CALL_AM33CALL    = 0x12, /* AM33 call */
+  PDB_CV_CALL_TRICALL     = 0x13, /* TriCore Call */
+  PDB_CV_CALL_SH5CALL     = 0x14, /* Hitachi SuperH-5 call */
+  PDB_CV_CALL_M32RCALL    = 0x15, /* M32R Call */
+  PDB_CV_CALL_CLRCALL     = 0x16, /* clr call */
+  PDB_CV_CALL_INLINE      = 0x17, /* Marker for routines always inlined and thus lacking a convention */
+  PDB_CV_CALL_NEAR_VECTOR = 0x18, /* near left to right push with regs, callee pops stack */
+  PDB_CV_CALL_RESERVED    = 0x19  /* first unused call enumeration */
+
+  /* Do NOT add any more machine specific conventions.  This is to be used for
+   * calling conventions in the source only (e.g. __cdecl, __stdcall). */
+} pdb_cv_call_e;
+
+#define PDB_CV_SIG_C6 0
+
+#define PDB_CV_SIG_C7 1
+/* C11 (vc5.x) */
+#define PDB_CV_SIG_C11 2
+/* C13 (vc7.x) */
+#define PDB_CV_SIG_C13 4 
+
+#define PDB_CV_SIG_RESERVED 5
+
+#define PDB_CV_SYM_LIST \
+  X(COMPILE, 0x0001) /* Compile flags symbol */ \
+  X(REGISTER_16t, 0x0002) /* Register variable */ \
+  X(CONSTANT_16t, 0x0003) /* constant symbol */ \
+  X(UDT_16t, 0x0004) /* User defined type */ \
+  X(SSEARCH, 0x0005) /* Start Search */ \
+  X(END, 0x0006) /* Block, procedure, "with" or thunk end */ \
+  X(SKIP, 0x0007) /* Reserve symbol space in $$Symbols table */ \
+  X(CVRESERVE, 0x0008) /* Reserved symbol for CV internal use */ \
+  X(OBJNAME_ST, 0x0009) /* path to object file name */ \
+  X(ENDARG, 0x000a) /* end of argument/return list */ \
+  X(COBOLUDT_16t, 0x000b) /* special UDT for cobol that does not symbol pack */ \
+  X(MANYREG_16t, 0x000c) /* multiple register variable */ \
+  X(RETURN, 0x000d) /* return description symbol */ \
+  X(ENTRYTHIS, 0x000e) /* description of this pointer on entry */ \
+  X(BPREL16, 0x0100) /* BP-relative */ \
+  X(LDATA16, 0x0101) /* Module-static symbol */ \
+  X(GDATA16, 0x0102) /* Global data symbol */ \
+  X(PUB16, 0x0103) /* a public symbol */ \
+  X(LPROC16, 0x0104) /* Local procedure start */ \
+  X(GPROC16, 0x0105) /* Global procedure start */ \
+  X(THUNK16, 0x0106) /* Thunk Start */ \
+  X(BLOCK16, 0x0107) /* block start */ \
+  X(WITH16, 0x0108) /* with start */ \
+  X(LABEL16, 0x0109) /* code label */ \
+  X(CEXMODEL16, 0x010a) /* change execution model */ \
+  X(VFTABLE16, 0x010b) /* address of virtual function table */ \
+  X(REGREL16, 0x010c) /* register relative address */ \
+  X(BPREL32_16t, 0x0200) /* BP-relative */ \
+  X(LDATA32_16t, 0x0201) /* Module-static symbol */ \
+  X(GDATA32_16t, 0x0202) /* Global data symbol */ \
+  X(PUB32_16t, 0x0203) /* a public symbol (CV internal reserved) */ \
+  X(LPROC32_16t, 0x0204) /* Local procedure start */ \
+  X(GPROC32_16t, 0x0205) /* Global procedure start */ \
+  X(THUNK32_ST, 0x0206) /* Thunk Start */ \
+  X(BLOCK32_ST, 0x0207) /* block start */ \
+  X(WITH32_ST, 0x0208) /* with start */ \
+  X(LABEL32_ST, 0x0209) /* code label */ \
+  X(CEXMODEL32, 0x020a) /* change execution model */ \
+  X(VFTABLE32_16t, 0x020b) /* address of virtual function table */ \
+  X(REGREL32_16t, 0x020c) /* register relative address */ \
+  X(LTHREAD32_16t, 0x020d) /* static thread storage */ \
+  X(GTHREAD32_16t, 0x020e) /* static thread storage */ \
+  X(SLINK32, 0x020f) /* static link for MIPS EH implementation */ \
+  X(LPROCMIPS_16t, 0x0300) /* Local procedure start */ \
+  X(GPROCMIPS_16t, 0x0301) /* Global procedure start */ \
+  X(PROCREF_ST, 0x0400) /* Reference to a procedure */ \
+  X(DATAREF_ST, 0x0401) /* Reference to data */ \
+  X(ALIGN, 0x0402) /* Used for page alignment of symbols */ \
+  X(LPROCREF_ST, 0x0403) /* Local Reference to a procedure */ \
+  X(OEM, 0x0404) /* OEM defined symbol */ \
+  X(TI16_MAX, 0x1000) \
+  X(REGISTER_ST, 0x1001) /* Register variable */ \
+  X(CONSTANT_ST, 0x1002) /* constant symbol */ \
+  X(UDT_ST, 0x1003) /* User defined type */ \
+  X(COBOLUDT_ST, 0x1004) /* special UDT for cobol that does not symbol pack */ \
+  X(MANYREG_ST, 0x1005) /* multiple register variable */ \
+  X(BPREL32_ST, 0x1006) /* BP-relative */ \
+  X(LDATA32_ST, 0x1007) /* Module-static symbol */ \
+  X(GDATA32_ST, 0x1008) /* Global data symbol */ \
+  X(PUB32_ST, 0x1009) /* a public symbol (CV internal reserved) */ \
+  X(LPROC32_ST, 0x100a) /* Local procedure start */ \
+  X(GPROC32_ST, 0x100b) /* Global procedure start */ \
+  X(VFTABLE32, 0x100c) /* address of virtual function table */ \
+  X(REGREL32_ST, 0x100d) /* register relative address */ \
+  X(LTHREAD32_ST, 0x100e) /* static thread storage */ \
+  X(GTHREAD32_ST, 0x100f) /* static thread storage */ \
+  X(LPROCMIPS_ST, 0x1010) /* Local procedure start */ \
+  X(GPROCMIPS_ST, 0x1011) /* Global procedure start */ \
+  X(FRAMEPROC, 0x1012) /* extra frame and proc information */ \
+  X(COMPILE2_ST, 0x1013) /* extended compile flags and info */ \
+  X(MANYREG2_ST, 0x1014) /* multiple register variable */ \
+  X(LPROCIA64_ST, 0x1015) /* Local procedure start (IA64) */ \
+  X(GPROCIA64_ST, 0x1016) /* Global procedure start (IA64) */ \
+  X(LOCALSLOT_ST, 0x1017) /* static IL sym with field for static slot index */ \
+  X(PARAMSLOT_ST, 0x1018) /* static IL sym with field for parameter slot index */ \
+  X(GMANPROC_ST, 0x101a) /* Global proc */ \
+  X(LMANPROC_ST, 0x101b) /* Local proc */ \
+  X(RESERVED1, 0x101c) /* reserved */ \
+  X(RESERVED2, 0x101d) /* reserved */ \
+  X(RESERVED3, 0x101e) /* reserved */ \
+  X(RESERVED4, 0x101f) /* reserved */ \
+  X(LMANDATA_ST, 0x1020) /**/   \
+  X(GMANDATA_ST, 0x1021) /**/ \
+  X(MANFRAMEREL_ST, 0x1022) /**/ \
+  X(MANREGISTER_ST, 0x1023) /**/ \
+  X(MANSLOT_ST, 0x1024) /**/ \
+  X(MANMANYREG_ST, 0x1025) /**/ \
+  X(MANREGREL_ST, 0x1026) /**/ \
+  X(MANMANYREG2_ST, 0x1027) /**/ \
+  X(MANTYPREF, 0x1028) /* Index for type referenced by name from metadata */ \
+  X(UNAMESPACE_ST, 0x1029) /* Using namespace */ \
+  X(ST_MAX, 0x1100) /* starting point for SZ name symbols */ \
+  X(OBJNAME, 0x1101) /* path to object file name */ \
+  X(THUNK32, 0x1102) /* Thunk Start */ \
+  X(BLOCK32, 0x1103) /* block start */ \
+  X(WITH32, 0x1104) /* with start */ \
+  X(LABEL32, 0x1105) /* code label */ \
+  X(REGISTER, 0x1106) /* Register variable */ \
+  X(CONSTANT, 0x1107) /* constant symbol */ \
+  X(UDT, 0x1108) /* User defined type */ \
+  X(COBOLUDT, 0x1109) /* special UDT for cobol that does not symbol pack */ \
+  X(MANYREG, 0x110a) /* multiple register variable */ \
+  X(BPREL32, 0x110b) /* BP-relative */ \
+  X(LDATA32, 0x110c) /* Module-static symbol */ \
+  X(GDATA32, 0x110d) /* Global data symbol */ \
+  X(PUB32, 0x110e) /* a public symbol (CV internal reserved) */ \
+  X(LPROC32, 0x110f) /* Local procedure start */ \
+  X(GPROC32, 0x1110) /* Global procedure start */ \
+  X(REGREL32, 0x1111) /* register relative address */ \
+  X(LTHREAD32, 0x1112) /* static thread storage */ \
+  X(GTHREAD32, 0x1113) /* static thread storage */ \
+  X(LPROCMIPS, 0x1114) /* Local procedure start */ \
+  X(GPROCMIPS, 0x1115) /* Global procedure start */ \
+  X(COMPILE2, 0x1116) /* extended compile flags and info */ \
+  X(MANYREG2, 0x1117) /* multiple register variable */ \
+  X(LPROCIA64, 0x1118) /* Local procedure start (IA64) */ \
+  X(GPROCIA64, 0x1119) /* Global procedure start (IA64) */ \
+  X(LOCALSLOT, 0x111a) /* static IL sym with field for static slot index */ \
+  X(PARAMSLOT, 0x111b) /* static IL sym with field for parameter slot index */ \
+  X(LMANDATA, 0x111c) /**/ \
+  X(GMANDATA, 0x111d) /**/ \
+  X(MANFRAMEREL, 0x111e) /**/ \
+  X(MANREGISTER, 0x111f) /**/ \
+  X(MANSLOT, 0x1120) /**/ \
+  X(MANMANYREG, 0x1121) /**/ \
+  X(MANREGREL, 0x1122) /**/ \
+  X(MANMANYREG2, 0x1123) /**/ \
+  X(UNAMESPACE, 0x1124) /* Using namespace */ \
+  X(PROCREF, 0x1125) /* Reference to a procedure */ \
+  X(DATAREF, 0x1126) /* Reference to data */ \
+  X(LPROCREF, 0x1127) /* Local Reference to a procedure */ \
+  X(ANNOTATIONREF, 0x1128) /* Reference to an S_ANNOTATION symbol */ \
+  X(TOKENREF, 0x1129) /* Reference to one of the many MANPROCSYM's */ \
+  X(GMANPROC, 0x112a) /* Global proc */ \
+  X(LMANPROC, 0x112b) /* Local proc */ \
+  X(TRAMPOLINE, 0x112c) /* trampoline thunks */ \
+  X(MANCONSTANT, 0x112d) /* constants with metadata type info */ \
+  X(ATTR_FRAMEREL, 0x112e) /* relative to virtual frame ptr */ \
+  X(ATTR_REGISTER, 0x112f) /* stored in a register */ \
+  X(ATTR_REGREL, 0x1130) /* relative to register (alternate frame ptr) */ \
+  X(ATTR_MANYREG, 0x1131) /* stored in >1 register */ \
+  X(SEPCODE, 0x1132)  \
+  X(LOCAL_2005, 0x1133) /* defines a static symbol in optimized code */ \
+  X(DEFRANGE_2005, 0x1134) /* defines a single range of addresses in which symbol can be evaluated */ \
+  X(DEFRANGE2_2005, 0x1135) /* defines ranges of addresses in which symbol can be evaluated */ \
+  X(SECTION, 0x1136) /* A COFF section in a PE executable */ \
+  X(COFFGROUP, 0x1137) /* A COFF group */ \
+  X(EXPORT, 0x1138) /* A export */ \
+  X(CALLSITEINFO, 0x1139) /* Indirect call site information */ \
+  X(FRAMECOOKIE, 0x113a) /* Security cookie information */ \
+  X(DISCARDED, 0x113b) /* Discarded by LINK /OPT:REF (experimental, see richards) */ \
+  X(COMPILE3, 0x113c) /* Replacement for S_COMPILE2 */ \
+  X(ENVBLOCK, 0x113d) /* Environment block split off from S_COMPILE2 */ \
+  X(LOCAL, 0x113e) /* defines a static symbol in optimized code */ \
+  X(DEFRANGE, 0x113f) /* defines a single range of addresses in which symbol can be evaluated */ \
+  X(DEFRANGE_SUBFIELD, 0x1140) /* ranges for a subfield */ \
+  X(DEFRANGE_REGISTER, 0x1141) /* ranges for en-registered symbol */ \
+  X(DEFRANGE_FRAMEPOINTER_REL, 0x1142) /* range for stack symbol. */ \
+  X(DEFRANGE_SUBFIELD_REGISTER, 0x1143) /* ranges for en-registered field of symbol */ \
+  X(DEFRANGE_FRAMEPOINTER_REL_FULL_SCOPE, 0x1144) /* range for stack symbol span valid full scope of function body, gap might apply. */ \
+  X(DEFRANGE_REGISTER_REL, 0x1145) /* range for symbol address as register + offset. */ \
+  X(LPROC32_ID, 0x1146) /**/ \
+  X(GPROC32_ID, 0x1147) /**/ \
+  X(LPROCMIPS_ID, 0x1148) /**/ \
+  X(GPROCMIPS_ID, 0x1149) /**/ \
+  X(LPROCIA64_ID, 0x114a) /**/ \
+  X(GPROCIA64_ID, 0x114b) /**/ \
+  X(BUILDINFO, 0x114c) /* build information. */ \
+  X(INLINESITE, 0x114d) /* inlined function callsite. */ \
+  X(INLINESITE_END, 0x114e) /**/ \
+  X(PROC_ID_END, 0x114f) /**/ \
+  X(DEFRANGE_HLSL, 0x1150) /**/ \
+  X(GDATA_HLSL, 0x1151) /**/ \
+  X(LDATA_HLSL, 0x1152) /**/ \
+  X(FILESTATIC, 0x1153) \
+  X(LOCAL_DPC_GROUPSHARED, 0x1154) /* DPC groupshared variable */ \
+  X(LPROC32_DPC, 0x1155) /* DPC static procedure start */ \
+  X(LPROC32_DPC_ID, 0x1156) /**/ \
+  X(DEFRANGE_DPC_PTR_TAG, 0x1157) /* DPC pointer tag definition range */ \
+  X(DPC_SYM_TAG_MAP, 0x1158) /* DPC pointer tag value to symbol record map */ \
+  X(ARMSWITCHTABLE, 0x1159) /**/ \
+  X(CALLEES, 0x115a) /**/ \
+  X(CALLERS, 0x115b) /**/ \
+  X(POGODATA, 0x115c) /**/ \
+  X(INLINESITE2, 0x115d) /* extended inline site information */ \
+  X(HEAPALLOCSITE, 0x115e) /* heap allocation site */ \
+  X(MOD_TYPEREF, 0x115f) /* only generated at link time */ \
+  X(REF_MINIPDB, 0x1160) /* only generated at link time for mini PDB */ \
+  X(PDBMAP, 0x1161) /* only generated at link time for mini PDB */ \
+  X(GDATA_HLSL32, 0x1162) /* */ \
+  X(LDATA_HLSL32, 0x1163) /* */ \
+  X(GDATA_HLSL32_EX, 0x1164) /**/ \
+  X(LDATA_HLSL32_EX, 0x1165) /**/ \
+  X(FASTLINK, 0x1167) \
+  X(INLINEES, 0x1168)
+
+typedef enum {
+  PDB_CV_SYM_NULL,
+
+#define X(name, value) PDB_CV_SYM_##name = value,
+  PDB_CV_SYM_LIST
+#undef X
+
+  PDB_CV_SYM_COUNT
+} pdb_cv_sym_type_e;
+typedef U16 pdb_cv_sym_type;
+typedef U16 pdb_cv_sym_size;
+
+#define PDB_CV_INLINEE_SOURCE_LINE_SIGNATURE     0x0
+#define PDB_CV_INLINEE_SOURCE_LINE_SIGNATURE_EX  0x1
+typedef struct pdb_cv_inlinee_srcline 
+{
+  pdb_cv_itemid inlinee;  /* Function ID */
+  U32       file_id;  /* offset into PDB_CV_SS_TYPE_FILE_CHKSUM */
+  U32       src_ln;   /* Line number of inlinee in file. */
+} pdb_cv_inlinee_srcline;
+
+typedef struct pdb_cv_inlinee_srcline_ex 
+{
+  pdb_cv_itemid inlinee; /* Function ID */
+  U32 file_id;       /* offset into PDB_CV_SS_TYPE_FILE_CHKSUM */
+  U32 src_ln;        /* Line number of inlinee in file. */
+  U32 extra_file_id_count;
+  /* U32     extra_file_id[0]; */
+} pdb_cv_inlinee_srcline_ex;
+
+typedef struct pdb_cv_symref2 
+{ /* PDB_CV_SYM_PROCREF, PDB_CV_SYM_DATAREF, or PDB_CV_SYM_LPROCREF */
+  U32 sum_name;   /* SUC of the name */
+  U32 sym_off;    /* Offset of actual symbol in $$Symbols */
+  pdb_imod imod;       /* Module containing the actual symbol */
+  /* unsigned char   name[1];     hidden name made a first class member */
+} pdb_cv_symref2;
+
+typedef struct pdb_cv_buildinfo
+{
+  pdb_cv_itemid id;
+} pdb_cv_buildinfo;
+
+typedef struct pdb_cv_funclist
+{
+  u32 count;
+  // pdb_cv_itype func[1];
+} pdb_cv_funclist;
+
+typedef struct pdb_cv_line_sec 
+{
+  pdb_isec_umm sec_off;
+  pdb_isec sec;
+  U16 flags;
+  U32 len;
+} pdb_cv_line_sec;
+
+typedef struct pdb_cv_src_file 
+{
+  U32 chksum_off;
+  U32 num_lines;
+  U32 lines_size;      /*  Size of pdb_cv_line record. */
+} pdb_cv_src_file;
+
+#define PDB_CV_LINE_GET_LN(line) ((line).flags & PDB_CV_LINE_ENUM_START)
+/* line where statement/expression starts. */
+#define PDB_CV_LINE_ENUM_START 0x00ffffff
+/* delta to line where statement ends (optional). */
+#define PDB_CV_LINE_DELTA_END 0x7f000000
+/* syms_true if a statement line number, else an expression line num. */
+#define PDB_CV_LINE_STATMENT 0x80000000
+
+typedef struct pdb_cv_line 
+{
+  U32 off;
+  U32 flags;
+} pdb_cv_line;
+
+typedef enum 
+{
+  PDB_CV_CHECKSUM_NULL,
+  PDB_CV_CHECKSUM_MD5,
+  PDB_CV_CHECKSUM_SHA1,
+  PDB_CV_CHECKSUM_SHA256
+} pdb_cv_chksum_type_e;
+typedef U8 pdb_cv_chksum_type;
+
+typedef struct pdb_cv_file_checksum 
+{
+  U32 name_off;
+  U8 len;
+  U8 type;
+} pdb_cv_file_checksum;
+
+#define PDB_CV_SS_TYPE_IGNORE(x) ((((U32)x) & 0x80000000) != 0)
+typedef enum 
+{
+  PDB_CV_SS_TYPE_SYMBOLS               = 0xf1,
+  PDB_CV_SS_TYPE_LINES                 = 0xf2,
+  PDB_CV_SS_TYPE_STRINGTABLE           = 0xf3,
+  PDB_CV_SS_TYPE_FILE_CHKSUM           = 0xf4,
+  PDB_CV_SS_TYPE_FRAMEDATA             = 0xf5,
+  PDB_CV_SS_TYPE_INLINE_LINES          = 0xf6,
+  PDB_CV_SS_TYPE_CROSS_SCOPE_IMPORTS   = 0xf7,
+  PDB_CV_SS_TYPE_CROSS_SCOPE_EXPORTS   = 0xf8,
+  PDB_CV_SS_TYPE_IL_LINES              = 0xf9,
+  PDB_CV_SS_TYPE_FUNC_MDTOKEN_MAP      = 0xfa,
+  PDB_CV_SS_TYPE_TYPE_MDTOKEN_MAP      = 0xfb,
+  PDB_CV_SS_TYPE_MERGED_ASSEMBLY_INPUT = 0xfc,
+  PDB_CV_SS_TYPE_COFF_SYMBOL_RVA       = 0xfe
+} pdb_cv_ss_type_e;
+typedef U32 pdb_cv_ss_type;
+
+typedef struct pdb_cv_proc 
+{
+  U32 parent;
+  U32 end;
+  U32 next;
+  U32 len;
+  U32 dbg_start;
+  U32 dbg_end;
+  pdb_cv_itype itype;
+  U32 off;
+  U16 sec;
+  U8 flags;     /* pdb_cv_proc_flags_e */
+} pdb_cv_proc;
+
+typedef struct pdb_cv_udtsym 
+{
+  pdb_cv_itype itype;
+  /* unsigned char name[0]; */
+} pdb_cv_udtsym;
+
+typedef enum 
+{
+  PDB_CV_BA_OP_END = 0,                      
+  PDB_CV_BA_OP_CODE_OFFSET = 1,              
+  PDB_CV_BA_OP_CHANGE_CODE_OFFSET_BASE = 2,  
+  PDB_CV_BA_OP_CHANGE_CODE_OFFSET = 3,      
+  PDB_CV_BA_OP_CHANGE_CODE_LENGTH = 4,      
+  PDB_CV_BA_OP_CHANGE_FILE = 5,            
+  PDB_CV_BA_OP_CHANGE_LINE_OFFSET = 6,      
+  PDB_CV_BA_OP_CHANGE_LINE_END_DELTA = 7,    
+  PDB_CV_BA_OP_CHANGE_RANGE_KIND = 8,       
+  PDB_CV_BA_OP_CHANGE_COLUMN_START = 9,      
+  PDB_CV_BA_OP_CHANGE_COLUMN_END_DELTA = 10,  
+  PDB_CV_BA_OP_CHANGE_CODE_OFFSET_AND_LINE_OFFSET = 11,  
+  PDB_CV_BA_OP_CHANGE_CODE_LENGTH_AND_CODE_OFFSET = 12,  
+  PDB_CV_BA_OP_CHANGE_COLUMN_END = 13,
+  PDB_CV_BA_OP_MAX = 0xff
+} pdb_cv_binary_annotation_opcode_e;
+typedef U8 pdb_cv_binary_annotation_opcode;
+
+enum
+{
+  PDB_CV_CFL_C        = 0x00,
+  PDB_CV_CFL_CXX      = 0x01,
+  PDB_CV_CFL_FORTRAN  = 0x02,
+  PDB_CV_CFL_MASM     = 0x03,
+  PDB_CV_CFL_PASCAL   = 0x04,
+  PDB_CV_CFL_BASIC    = 0x05,
+  PDB_CV_CFL_COBOL    = 0x06,
+  PDB_CV_CFL_LINK     = 0x07,
+  PDB_CV_CFL_CVTRES   = 0x08,
+  PDB_CV_CFL_CVTPGD   = 0x09,
+  PDB_CV_CFL_CSHARP   = 0x0A,
+  PDB_CV_CFL_VB       = 0x0B,
+  PDB_CV_CFL_ILASM    = 0x0C,
+  PDB_CV_CFL_JAVA     = 0x0D,
+  PDB_CV_CFL_JSCRIPT  = 0x0E,
+  PDB_CV_CFL_MSIL     = 0x0F,
+  PDB_CV_CFL_HLSL     = 0x10
+};
+
+enum
+{
+  PDB_CV_CFL_8080         = 0x00,
+  PDB_CV_CFL_8086         = 0x01,
+  PDB_CV_CFL_80286        = 0x02,
+  PDB_CV_CFL_80386        = 0x03,
+  PDB_CV_CFL_80486        = 0x04,
+  PDB_CV_CFL_PENTIUM      = 0x05,
+  PDB_CV_CFL_PENTIUMII    = 0x06,
+  PDB_CV_CFL_PENTIUMPRO   = PDB_CV_CFL_PENTIUMII,
+  PDB_CV_CFL_PENTIUMIII   = 0x07,
+  PDB_CV_CFL_MIPS         = 0x10,
+  PDB_CV_CFL_MIPSR4000    = PDB_CV_CFL_MIPS,
+  PDB_CV_CFL_MIPS16       = 0x11,
+  PDB_CV_CFL_MIPS32       = 0x12,
+  PDB_CV_CFL_MIPS64       = 0x13,
+  PDB_CV_CFL_MIPSI        = 0x14,
+  PDB_CV_CFL_MIPSII       = 0x15,
+  PDB_CV_CFL_MIPSIII      = 0x16,
+  PDB_CV_CFL_MIPSIV       = 0x17,
+  PDB_CV_CFL_MIPSV        = 0x18,
+  PDB_CV_CFL_M68000       = 0x20,
+  PDB_CV_CFL_M68010       = 0x21,
+  PDB_CV_CFL_M68020       = 0x22,
+  PDB_CV_CFL_M68030       = 0x23,
+  PDB_CV_CFL_M68040       = 0x24,
+  PDB_CV_CFL_ALPHA        = 0x30,
+  PDB_CV_CFL_ALPHA_21064  = 0x30,
+  PDB_CV_CFL_ALPHA_21164  = 0x31,
+  PDB_CV_CFL_ALPHA_21164A = 0x32,
+  PDB_CV_CFL_ALPHA_21264  = 0x33,
+  PDB_CV_CFL_ALPHA_21364  = 0x34,
+  PDB_CV_CFL_PPC601       = 0x40,
+  PDB_CV_CFL_PPC603       = 0x41,
+  PDB_CV_CFL_PPC604       = 0x42,
+  PDB_CV_CFL_PPC620       = 0x43,
+  PDB_CV_CFL_PPCFP        = 0x44,
+  PDB_CV_CFL_PPCBE        = 0x45,
+  PDB_CV_CFL_SH3          = 0x50,
+  PDB_CV_CFL_SH3E         = 0x51,
+  PDB_CV_CFL_SH3DSP       = 0x52,
+  PDB_CV_CFL_SH4          = 0x53,
+  PDB_CV_CFL_SHMEDIA      = 0x54,
+  PDB_CV_CFL_ARM3         = 0x60,
+  PDB_CV_CFL_ARM4         = 0x61,
+  PDB_CV_CFL_ARM4T        = 0x62,
+  PDB_CV_CFL_ARM5         = 0x63,
+  PDB_CV_CFL_ARM5T        = 0x64,
+  PDB_CV_CFL_ARM6         = 0x65,
+  PDB_CV_CFL_ARM_XMAC     = 0x66,
+  PDB_CV_CFL_ARM_WMMX     = 0x67,
+  PDB_CV_CFL_ARM7         = 0x68,
+  PDB_CV_CFL_OMNI         = 0x70,
+  PDB_CV_CFL_IA64         = 0x80,
+  PDB_CV_CFL_IA64_1       = 0x80,
+  PDB_CV_CFL_IA64_2       = 0x81,
+  PDB_CV_CFL_CEE          = 0x90,
+  PDB_CV_CFL_AM33         = 0xA0,
+  PDB_CV_CFL_M32R         = 0xB0,
+  PDB_CV_CFL_TRICORE      = 0xC0,
+  PDB_CV_CFL_X64          = 0xD0,
+  PDB_CV_CFL_AMD64        = PDB_CV_CFL_X64,
+  PDB_CV_CFL_EBC          = 0xE0,
+  PDB_CV_CFL_THUMB        = 0xF0,
+  PDB_CV_CFL_ARMNT        = 0xF4,
+  PDB_CV_CFL_ARM64        = 0xF6,
+  PDB_CV_CFL_D3D11_SHADER = 0x100
+};
+
+#define PDB_CV_COMPILESYM2_MASK_LANG 0xFF
+#define PDB_CV_COMPILESYM2_FLAG_EC (1u << 8u)
+#define PDB_CV_COMPILESYM2_FLAG_NODBGINFO (1u << 9u)
+#define PDB_CV_COMPILESYM2_FLAG_LTCG (1u << 10u)
+#define PDB_CV_COMPILESYM2_FLAG_NODATAALIGN (1u << 11u)
+#define PDB_CV_COMPILESYM2_FLAG_MANAGED_PRESENT (1u << 12u)
+#define PDB_CV_COMPILESYM2_FLAG_SECURITY_CHECKS (1u << 13u)
+#define PDB_CV_COMPILESYM2_FLAG_HOTPATCH (1u << 14u)
+#define PDB_CV_COMPILESYM2_FLAG_CVTCIL (1u << 15u)
+#define PDB_CV_COMPILESYM2_FLAG_MSIL_MODULE (1u << 16u)
+typedef u32 pdb_cv_compilesym2_flags;
+
+typedef struct pdb_cv_compilesym2
+{
+  pdb_cv_compilesym2_flags flags;
+  u16 machine;
+  u16 ver_fe_major;
+  u16 ver_fe_minor;
+  u16 ver_fe_build;
+  u16 ver_major;
+  u16 ver_minor;
+  u16 ver_build;
+  /* char ver_str[0]; */
+} pdb_cv_compilesym2;
+
+#define PDB_CV_COMPILESYM3_PAD_SHIFT 20
+#define PDB_CV_COMPILESYM3_PAD_MASK 0xFFC00000
+#define PDB_CV_COMPILESYM3_LANG_MASK 0xFF
+// Compiled for E/C
+#define PDB_CV_COMPILESYM3_FLAGS_EC               (1u << 9u)
+// Compiled without debug info
+#define PDB_CV_COMPILESYM3_FLAGS_NODBGINFO        (1u << 10u)
+// Compiled with /LTCG (Link-time code generation)
+#define PDB_CV_COMPILESYM3_FLAGS_LTCG             (1u << 11u)
+// Compiled with -Bzalign (what is this flag?)
+#define PDB_CV_COMPILESYM3_FLAGS_NODATAALIGN      (1u << 12u)
+// Managed code/data is present
+#define PDB_CV_COMPILESYM3_FLAGS_MANAGED_PRESENT  (1u << 13u)
+// Compiled with /GS (Buffer Security Checks)
+#define PDB_CV_COMPILESYM3_FLAGS_SECURITY_CHECKS  (1u << 14u)
+// Compiled with /hotpatch (Create Hotpatchable Image)
+#define PDB_CV_COMPILESYM3_FLAGS_HOTPATCH         (1u << 15u)
+// Converted with CVTCIL
+#define PDB_CV_COMPILESYM3_FLAGS_CVTCIL           (1u << 16u)
+// Compiled with /NOASSEMBLY MSIL netmodule
+#define PDB_CV_COMPILESYM3_FLAGS_MSIL_MODULE      (1u << 17u)
+// Compiled with /sdl (Enable Additional Security Checks)
+#define PDB_CV_COMPILESYM3_FLAGS_SDL              (1u << 18u)
+// Compiled with /ltcg:pgo or pgu
+#define PDB_CV_COMPILESYM3_FLAGS_PGO              (1u << 19u)
+// .exp module
+#define PDB_CV_COMPILESYM3_FLAGS_EXP              (1u << 20u)
+
+typedef struct pdb_cv_compilesym3
+{
+  U32  flags;
+
+  /* target processor */
+  U16  machine;    
+
+  /* front end major version # */
+  U16  ver_fe_major; 
+
+  /* front end minor version # */
+  U16  ver_fe_minor;
+
+  /* front end build version # */
+  U16  ver_fe_build; 
+
+  /* front end QFE version # */
+  U16  ver_feqfe;   
+
+  /* back end major version # */
+  U16  ver_major;   
+
+  /* back end minor version # */
+  U16  ver_minor;   
+
+  /* back end build version # */
+  U16  ver_build;   
+
+  /* back end QFE version # */
+  U16  ver_qfe;     
+
+  /* char version [0]; Zero terminated compiler version string */
+} pdb_cv_compilesym3;
+
+typedef struct pdb_cv_sectionsym
+{
+  u16 sec_index;
+  u8 align;
+  u8 pad;
+  u32 rva;
+  u32 size;
+  u32 characteristics;
+  /* char name[0]; */
+} pdb_cv_sectionsym;
+
+typedef struct pdb_cv_coffgroupsym
+{
+  u32 size;
+  u32 characteristics;
+  u32 off;
+  u16 sec;
+  /* char name[0]; */
+} pdb_cv_coffgroupsym;
+
+enum
+{
+  PDB_CV_TRAMPOLINE_INCREMENTAL   = 0,
+  PDB_CV_TRAMPOLINE_BRANCH_ISLAND = 1
+};
+typedef u16 pdb_cv_trampoline_type;
+
+typedef struct pdb_cv_trampolinesym 
+{
+  /* trampoline symbol sub-type  */
+  pdb_cv_trampoline_type type;  
+
+  /* size of the thunk */
+  U16 thunk_size;
+
+  /* offset of the thunk */
+  U32 thunk_sec_off;   
+
+  /* offset of the target of the thunk */
+  U32 target_sec_off;  
+
+  /* section index of the thunk */
+  pdb_isec thunk_sec;  
+
+  /* section index of the target of the thunk */
+  pdb_isec target_sec; 
+} pdb_cv_trampolinesym;
+
+typedef struct pdb_cv_thunksym32
+{
+  u32 parent;
+  u32 end;
+  u32 next;
+  u32 off;
+  u16 sec;
+  u16 len;
+  u8 ord;
+  // char name[0];
+  // char variant[0];
+} pdb_cv_thunksym32;
+
+typedef struct pdb_cv_objnamesym 
+{
+  U32 sig;
+  /* U8 name[1]; ; Length-prefixed name */
+} pdb_cv_objnamesym;
+
+#define PDB_CV_ENVBLOCK_FLAG_REV  (1u << 0u)
+#define PDB_CV_ENVBLOCK_MASK_PAD  0x7f
+typedef u8 pdb_cv_envblocksym_flags;
+
+typedef struct pdb_cv_envblocksym
+{
+  pdb_cv_envblocksym_flags flags;
+  /* char str[0]; */
+} pdb_cv_envblocksym;
+
+typedef struct pdb_cv_inlinesym 
+{
+  /* pointer to the inliner */
+  U32 parent_offset;   
+
+  /* pointer to this block's end */
+  U32 end_offset;      
+
+  pdb_cv_itemid inlinee;   
+
+  /* an array of compressed binary annotations. */
+  /* pdb_u8 binaryAnnotations[0]; */ 
+} pdb_cv_inlinesym;
+
+typedef struct pdb_cv_inlinesym2 
+{
+  U32 parent_offset;
+  U32 end_offset;
+  pdb_cv_itemid inlinee;
+  U32 invocations;
+  /* an array of compressed binary annotations. */
+  /* pdb_u8 binaryAnnotations[0]; */ 
+} pdb_cv_inlinesym2;
+
+#if 0
+typedef struct pdb_cv_inlineesym
+{
+  u32 count;
+  pdb_itype funcs[0];
+} pdb_cv_inlineesym;
+#endif
+
+enum 
+{
+  PDB_CV_FASTLINK_FLAGS_IS_GLOBAL_DATA = (1 << 0),
+  PDB_CV_FASTLINK_FALGS_IS_DATA = (1 << 1),
+  PDB_CV_FASTLINK_FLAGS_IS_UDT = (1 << 2),
+  PDB_CV_FASTLINK_FLAGS_IS_CONST = (1 << 4),
+  PDB_CV_FASTLINK_FLAGS_IS_NAMESPACE = (1 << 6)
+};
+typedef U16 pdb_cv_fastlink_flags;
+
+typedef struct pdb_cv_fastlink 
+{
+  union {
+    /* NOTE(nick): If PDB_CV_FASTLINK_FALGS_IS_UDT is set then this is a type index */
+    pdb_cv_itype type_index;
+    U32 unknown;
+  } u;
+  pdb_cv_fastlink_flags flags;
+
+  /* U8 name[0] */
+  /* U8 padding[4] */
+} pdb_cv_fastlink;
+
+typedef struct pdb_cv_datasym32 
+{
+  /* Type index, or Metadata token if a managed symbol */
+  pdb_cv_itype itype;         
+
+  pdb_isec_umm sec_off;
+  pdb_isec sec;
+
+  /* unsigned char   name[1]; */
+} pdb_cv_datasym32;
+
+typedef struct pdb_cv_constsym 
+{
+  /* Type index (containing enum if enumerate) or metadata token */
+  pdb_cv_itype itype;     
+
+  /* numeric leaf containing value; use pdb_stream_read_numeric() */
+
+  /* unsigned char   name[CV_ZEROLEN];     / Length-prefixed name */
+} pdb_cv_constsym;
+
+typedef struct pdb_cv_regrel32 
+{
+  /* offset of symbol */
+  U32 reg_off;    
+
+  /* Type index or metadata token */
+  pdb_cv_itype itype;      
+
+  /* register index for symbol */
+  U16 reg;        
+
+  /* char   name[]; */
+} pdb_cv_regrel32;
+
+typedef struct pdb_cv_blocksym32 
+{
+  /* pointer to the parent */
+  U32 par;       
+
+  /* pointer to this blocks end */
+  U32 end;       
+
+  /* Block length */
+  U32 len;       
+
+  /* Offset in code section */
+  U32 off;       
+
+  /* section index */
+  U16 sec;       
+
+  /* Length-prefixed name */
+  /* unsigned char   name[1]; */
+} pdb_cv_blocksym32;
+
+typedef enum 
+{
+  /* Frame pointer present */
+  PDB_CV_PROC32_FLAG_NOFPO = (1 << 0),
+
+  /* Interrupt return */
+  PDB_CV_PROC32_FLAG_INT = (1 << 1),
+
+  PDB_CV_PROC32_FLAG_FAR = (1 << 2),  /* far return */
+
+  /* Function does not return */
+  PDB_CV_PROC32_FLAG_NEVER = (1 << 3),
+
+  /* Label isn't falen into */
+  PDB_CV_PROC32_FLAG_NOTREACHED = (1 << 4),
+
+  /* Custom call convention */
+  PDB_CV_PROC32_FLAG_CUSTOM_CALL = (1 << 5),
+
+  /* Function marked as noinline. */
+  PDB_CV_PROC32_FLAG_NOINLINE = (1 << 6),
+
+  /* Function has debug info for optimized code. */
+  PDB_CV_PROC32_FLAG_OPTDBGINFO = (1 << 7)
+} pdb_cv_procsym32_flags_e; 
+typedef u8 pdb_cv_proc_flags;
+
+typedef struct pdb_cv_procsym32 
+{
+  /* pointer to the parent */
+  U32 parent;     
+
+  /* pointer to this blocks end */
+  U32 end;        
+
+  /* pointer to next symbol */
+  U32 next;       
+
+  /* Proc length */
+  U32 len;        
+
+  /* Debug start offset */
+  U32 dbg_start;  
+
+  /* Debug end offset */
+  U32 dbg_end;    
+
+  /* Type index or ID */
+  pdb_cv_itype itype;      
+
+  U32 off;
+  U16 sec;
+  pdb_cv_proc_flags flags;
+
+  /* Length-prefixed name */
+  /* unsigned char name[1]; */    
+} pdb_cv_procsym32;
+
+#define PDB_CV_EXPORT_FLAG_CONSTAN (1u << 0u)
+#define PDB_CV_EXPORT_FLAG_DATA (1u << 1u)
+#define PDB_CV_EXPORT_FLAG_PRIVATE (1u << 2u)
+#define PDB_CV_EXPORT_FLAG_NONAME (1u << 3u)
+#define PDB_CV_EXPORT_FLAG_ORDINAL (1u << 4u)
+#define PDB_CV_EXPORT_FLAG_FORWARDER (1u << 5u)
+typedef u16 pdb_cv_export_flags;
+
+typedef struct pdb_cv_exportsym
+{
+  u16 ordinal;
+  pdb_cv_export_flags flags;
+  //char name[0];
+} pdb_cv_exportsym;
+
+typedef struct pdb_cv_label32
+{
+  u32 off;
+  u16 seg;
+  pdb_cv_proc_flags flags;
+  char name[1];
+} pdb_cv_label32;
+
+typedef struct pdb_cv_unamespace
+{
+  char name[1];
+} pdb_cv_unamespace;
+
+typedef struct pdb_cv_callsiteinfo
+{
+  u32 off; // call site section offset
+  u16 sec; // call site section index
+  u16 pad; // padding
+  pdb_cv_itype itype; // function signature
+} pdb_cv_callsiteinfo;
+
+enum
+{
+  PDB_CV_COOKIE_COPY    = 0,
+  PDB_CV_COOKIE_XOR_SP  = 1,
+  PDB_CV_COOKIE_XOR_BP  = 2,
+  PDB_CV_COOKIE_XOR_R13 = 3 
+};
+typedef u32 pdb_cv_cookie_type;
+
+typedef struct pdb_cv_framecookie
+{
+  u32 off;
+  u16 reg;
+  pdb_cv_cookie_type type;
+  u8 flags;
+} pdb_cv_framecookie;
+
+typedef struct pdb_cv_heapallocsite
+{
+  u32 off; // call site offset
+  u16 sec; // call site section
+  u16 call_inst_len; // length of heap allocation call instruction (in bytes)
+  pdb_cv_itype itype; // function signature
+} pdb_cv_heapallocsite;
+
+typedef enum 
+{
+  /* Variable is a parameter. */
+  PDB_CV_LOCALSYM_FLAG_PARAM = (1 << 0),
+
+  /* Address of the variable is taken. */
+  PDB_CV_LOCALSYM_FLAG_ADDR_TAKEN = (1 << 1),
+
+  /* Variable is compiler generated. */
+  PDB_CV_LOCALSYM_FLAG_COPMGEN = (1 << 2),
+
+  /* Symbol is splitted in temporaries, which are treated by compiler as
+   * independent entities. */
+  PDB_CV_LOCALSYM_FLAG_AGGREGATE = (1 << 3),
+
+  /* Counterpart of LOCALSYM_FLAG_AGGREGATE - tells that it is part of a 
+   * aggregate symobl. */
+  PDB_CV_LOCALSYM_FLAG_PARTOF_ARGGREGATE = (1 << 4),
+
+  /* Variable has mutltiple simultaneous lifetimes. */
+  PDB_CV_LOCALSYM_FLAG_ALIASED = (1 << 5),
+
+  /* Represents one of the multiple simultaneous lifetimes. */
+  PDB_CV_LOCALSYM_FLAG_ALIAS = (1 << 6),
+
+  /* Represents a return value of the procedure. */
+  PDB_CV_LOCALSYM_FLAG_RETVAL = (1 << 7),
+
+  /* Variable was optimized out and has no lifetime. */
+  PDB_CV_LOCALSYM_FLAG_OPTOUT = (1 << 8),
+
+  /* Variable is registered as global. */
+  PDB_CV_LOCALSYM_FLAG_GLOBAL = (1 << 9),
+
+  /* Variable is registerd as static. */
+  PDB_CV_LOCALSYM_FLAG_STATIC = (1 << 10)
+} pdb_cv_localsym_flags_e;
+typedef U16 pdb_cv_localsym_flags;
+
+typedef struct pdb_cv_localsym 
+{
+  /* type index*/
+  pdb_cv_itype itype;
+
+  /* pdb_cv_localsym_flags_e */
+  U16 flags;      
+
+  /* Name of this symbol, a null terminated array of UTF8 characters.  */
+  /* unsigned char name[0]; */
+} pdb_cv_localsym;
+
+/* defines a range of addresses */
+typedef struct pdb_cv_lvar_addr_range 
+{
+  /* Relative to the section offset */
+  U32 off;     
+  U16 sec;
+
+  /* Byte count of the range */
+  U16 len;    
+} pdb_cv_lvar_addr_range_t;
+
+typedef struct pdb_cv_lvar_addr_gap 
+{
+  U16 off;
+  U16 len;
+} pdb_cv_lvar_addr_gap_t;
+
+/* May have no user name on one of control flow path. */
+#define PDB_CV_RANGE_ATTRIB_MAYBE(atrrib) (((attrib) & 1) != 0)
+
+typedef U16 pdb_cv_range_attrib_t;
+
+#define CV_OFF_PARENT_LEN_LIMIT(lim) (lim & (0xfff))
+typedef U32 pdb_cv_off_parent_len_limit_t;
+
+/* Spilled member for s.i.  */
+#define PDB_CV_DEFRANGE_REGREL_SPILLED_UDT_MEMBER(flags) (((flags) & (1 << 0)) != 0)
+/* Offset in parent variable. */
+#define PDB_CV_DEFRANGE_REGREL_PARENT_OFF(flags) (((flags) & 0xFFF0) >> 4)
+typedef struct pdb_cv_defrange_register_rel 
+{
+  /* Register to hold the base pointer of the * symbol */
+  U16 reg;                    
+
+  U16 flags;
+
+  /* Offset to register */
+  U32 reg_off;                
+
+  /* Range of addresses where this variable is valid  */
+  pdb_cv_lvar_addr_range_t range; 
+
+  /* pdb_cv_lvar_addr_gap_t   gaps[0]; The value is not available in following gaps. */
+} pdb_cv_defrange_register_rel;
+
+typedef struct pdb_cv_defrange_reg 
+{
+  U16 reg;
+  pdb_cv_range_attrib_t attrib;
+  pdb_cv_lvar_addr_range_t range;
+  /* pdb_cv_lvar_addr_gap_t gaps[0]; */
+} pdb_cv_defrange_reg;
+
+typedef struct pdb_cv_defrange_frameptr_rel 
+{ /* A live range of frame variable */
+  /* Offset to frame pointer */
+  U32 off; 
+
+  pdb_cv_lvar_addr_range_t range;
+
+  /* cv_lvar_addr_gap gaps[0]; */
+} pdb_cv_defrange_frameptr_rel;
+
+typedef struct pdb_cv_defrange_frameptr_rel_full_scope 
+{
+  /* Offset to frame pointer */
+  U32 off; 
+} pdb_cv_defrange_frameptr_rel_full_scope;
+
+typedef struct pdb_cv_defrange_subfield_reg 
+{
+  U16 reg;
+  pdb_cv_range_attrib_t attrib;
+  pdb_cv_off_parent_len_limit_t off_parent;
+  pdb_cv_lvar_addr_range_t range;
+  /* cv_lvar_addr_gap[]; */
+} pdb_cv_defrange_subfield_reg;
+
+typedef struct pdb_cv_objname 
+{
+  U32 sig;
+  /* char name[0] */
+} pdb_cv_objname;
+
+typedef struct pdb_cv_filestaticsym
+{
+  pdb_cv_itype itype;
+  u32 mod_offset;
+  pdb_cv_localsym_flags flags;
+  /* char name[0]; */
+} pdb_cv_filestaticsym;
+
+// Function uses _alloca()
+#define PDB_CV_FRAMEPROCSYM_FLAGS_HAS_ALLOCA (1u << 0u)
+// Function uses setjmp()
+#define PDB_CV_FRAMEPROCSYM_FLAGS_SET_JMP (1u << 1u)
+// Function uses longjmp()
+#define PDB_CV_FRAMEPROCSYM_FLAGS_LONG_JMP (1u << 2u)
+// Function uses inline asm
+#define PDB_CV_FRAMEPROCSYM_FLAGS_INL_ASM (1u << 3u)
+// Function uses EH states
+#define PDB_CV_FRAMEPROCSYM_FLAGS_HAS_EH (1u << 4u)
+// Function declared as inline
+#define PDB_CV_FRAMEPROCSYM_FLAGS_INLINE (1u << 5u)
+// Function has SEH
+#define PDB_CV_FRAMEPROCSYM_FLAGS_HAS_SEH (1u << 6u)
+// Function declared with __declspec(naked)
+#define PDB_CV_FRAMEPROCSYM_FLAGS_NAKED (1u << 7u)
+// Function contains security checks (/GS)
+#define PDB_CV_FRAMEPROCSYM_FLAGS_SECURITY_CHECKS (1u << 8u)
+// Function compiled with /EHa
+#define PDB_CV_FRAMEPROCSYM_FLAGS_ASYNC_EH (1u << 9u)
+// Compiled with security checks (/GS), but stack ordering was not done
+#define PDB_CV_FRAMEPROCSYM_FLAGS_GS_NOSTACKORDERING (1u << 10u)
+// Function was inlined within another function
+#define PDB_CV_FRAMEPROCSYM_FLAGS_WAS_INLINED (1u << 11u)
+// Function declared with __declspec(strict_gs_check)
+#define PDB_CV_FRAMEPROCSYM_FLAGS_SAFE_BUFFERS (1u << 12u)
+#define PDB_CV_FRAMEPROCSYM_FLAGS_ENCODED_LOCAL_BASE_POINTER_MASK 0x6000
+#define PDB_CV_FRAMEPROCSYM_FLAGS_ENCODED_LOCAL_BASE_POINTER_SHIFT 12
+#define PDB_CV_FRAMEPROCSYM_FLAGS_ENCODED_PARAM_BASE_POINTER_MASK 0x18000
+#define PDB_CV_FRAMEPROCSYM_FLAGS_ENCODED_PARAM_BASE_POINTER_SHIFT 14
+// Function was compiled with PGO/PGU
+#define PDB_CV_FRAMEPROCSYM_FLAGS_POGO_ON (1u << 18u)
+// Is POGO counts valid
+#define PDB_CV_FRAMEPROCSYM_FLAGS_VALID_COUNTS (1u << 19u)
+// Was function optimized for speed
+#define PDB_CV_FRAMEPROCSYM_FLAGS_OPT_SPEED (1u << 20u)
+// Function has CFG checks (and no write checks)
+#define PDB_CV_FRAMEPROCSYM_FLAGS_CFG (1u << 21u)
+// Function has CFW checks and/or instrumentation
+#define PDB_CV_FRAMEPROCSYM_FLAGS_CFW (1u << 22u)
+// Padding must be zero
+#define PDB_CV_FRAMEPROCSYM_FLAGS_PAD_MASK 0xFF800000 
+#define PDB_CV_FRAMEPROCSYM_FLAGS_PAD_SHIFT 22
+
+typedef struct pdb_cv_frameprocsym
+{
+  U32 frame_size;
+  U32 pad_size;
+  U32 pad_off; // frame pointer relative offset to where padding begins
+  U32 save_reg_size;
+  pdb_isec_umm eh_off; // offset of exception handler
+  pdb_isec eh_sec;
+  U32 flags;
+} pdb_cv_frameprocsym;
+typedef struct pdb_refsym2 
+{
+  /* SUC of the name */
+  U32 sub_name;    
+
+  /* Offset of actual symbol in $$Symbols */
+  U32 offset;      
+
+  /* Module containing the actual symbol */
+  pdb_imod imod;       
+
+  /* hidden name made a first class member */
+  /* unsigned char name[1]; */
+} pdb_refsym2;
+
+// what are these anomalies: 0x9, 0x15, 0x16, 0x17, 0x18, 0x19, 0x1a
+#define PDB_BASIC_TYPE_LIST \
+  X(NOTYPE,      "null",                  0x0) \
+  X(ABS,         "abs",                   0x1) \
+  X(SEGMENT,     "segment",               0x2) \
+  X(VOID,        "void",                  0x3) \
+  X(CURRENCY,    "currency",              0x4) \
+  X(NBASICSTR,   "nbasicstr",             0x5) \
+  X(FBASICSTR,   "fbasicstr",             0x6) \
+  X(NOTTRANS,    "nottrans",              0x7) \
+  X(HRESULT,     "HRESULT",               0x8) \
+  X(CHAR,        "char",                  0x10) \
+  X(SHORT,       "short",                 0x11) \
+  X(LONG,        "long",                  0x12) \
+  X(QUAD,        "long long",             0x13) \
+  X(OCT,         "octal",                 0x14) \
+  X(UCHAR,       "unsigned char",         0x20) \
+  X(USHORT,      "unsigned short",        0x21) \
+  X(ULONG,       "unsigned long",         0x22) \
+  X(UQUAD,       "unsigned long long",    0x23) \
+  X(UOCT,        "unsigned octal",        0x24) \
+  X(BOOL08,      "bool8",                 0x30) \
+  X(BOOL16,      "bool16",                0x31) \
+  X(BOOL32,      "bool32",                0x32) \
+  X(BOOL64,      "bool64",                0x33) \
+  X(REAL32,      "float32",               0x40) \
+  X(REAL64,      "float64",               0x41) \
+  X(REAL80,      "float80",               0x42) \
+  X(REAL128,     "float128",              0x43) \
+  X(REAL48,      "float48",               0x44) \
+  X(REAL32PP,    "float32pp",             0x45) \
+  X(REAL16,      "float16",               0x46) \
+  X(CPLX32,      "complex32",             0x50) \
+  X(CPLX64,      "complex64",             0x51) \
+  X(CPLX80,      "complex80",             0x52) \
+  X(CPLX128,     "complex128",            0x53) \
+  X(BIT,         "bit",                   0x60) \
+  X(PASCHAR,     "pscal_char",            0x61) \
+  X(BOOL32FF,    "bool32ff",              0x62) /* 32-bit BOOL where syms_true is 0xffffffff */ \
+  X(INT1,        "int8",                  0x68) \
+  X(UINT1,       "uint8",                 0x69) \
+  X(RCHAR,       "char",                  0x70) \
+  X(WCHAR,       "wchar_t",               0x71) \
+  X(INT2,        "int16",                 0x72) \
+  X(UINT2,       "uint16",                0x73) \
+  X(INT4,        "int32",                 0x74) \
+  X(UINT4,       "uint32",                0x75) \
+  X(INT8,        "int64",                 0x76) \
+  X(UINT8,       "uint64",                0x77) \
+  X(INT16,       "int128",                0x78) \
+  X(UINT16,      "uint128",               0x79) \
+  X(CHAR16,      "char16",                0x7a) \
+  X(CHAR32,      "char32",                0x7b) \
+  X(CHAR8,       "char8",                 0x7c) \
+  X(PTR,         "pointer",               0xf0) \
+  /* end of the list. */
+
+#define PDB_BASIC_TYPE_IS_PTR(x) ((((x) & 0x0f00) == 0x600) || (((x) & 0x0f00) == 0x500))
+#define PDB_BASIC_TYPE_SIZE_MASK(x) (((x) & 0x700) >> 8)
+#define PDB_BASIC_TYPE_KIND_MASK(x) ((x) & 0xff)
+#define PDB_BASIC_TYPE_SIZE_VALUE       0
+#define PDB_BASIC_TYPE_SIZE_16BIT       1
+#define PDB_BASIC_TYPE_SIZE_FAR_16BIT   2   /* 16:16 far pointer */
+#define PDB_BASIC_TYPE_SIZE_HUGE_16BIT  3   /* 16:16 huge pointer*/
+#define PDB_BASIC_TYPE_SIZE_32BIT       4
+#define PDB_BASIC_TYPE_SIZE_16_32BIT    5   /* 16:32 pointer */
+#define PDB_BASIC_TYPE_SIZE_64BIT       6
+typedef enum {
+#define X(name, display_name, value) PDB_BASIC_TYPE_##name = value,
+  PDB_BASIC_TYPE_LIST
+#undef X
+  PDB_BASIC_TYPE_MAX
+} pdb_basic_type_e;
+typedef u32 pdb_basic_type;
+
+typedef enum 
+{
+  /* leaf indices starting records but referenced from symbol records */
+  PDB_LF_MODIFIER_16t = 0x0001,
+  PDB_LF_POINTER_16t = 0x0002,
+  PDB_LF_ARRAY_16t = 0x0003,
+  PDB_LF_CLASS_16t = 0x0004,
+  PDB_LF_STRUCTURE_16t = 0x0005,
+  PDB_LF_UNION_16t = 0x0006,
+  PDB_LF_ENUM_16t = 0x0007,
+  PDB_LF_PROCEDURE_16t = 0x0008,
+  PDB_LF_MFUNCTION_16t = 0x0009,
+  PDB_LF_VTSHAPE = 0x000a,
+  PDB_LF_COBOL0_16t = 0x000b,
+  PDB_LF_COBOL1 = 0x000c,
+  PDB_LF_BARRAY_16t = 0x000d,
+  PDB_LF_LABEL = 0x000e,
+  PDB_LF_NULL = 0x000f,
+  PDB_LF_NOTTRAN = 0x0010,
+  PDB_LF_DIMARRAY_16t = 0x0011,
+  PDB_LF_VFTPATH_16t = 0x0012,
+  PDB_LF_PRECOMP_16t = 0x0013,            /* not referenced from symbol */
+  PDB_LF_ENDPRECOMP = 0x0014,             /* not referenced from symbol */
+  PDB_LF_OEM_16t = 0x0015,                /* oem definable type string */
+  PDB_LF_TYPESERVER_ST = 0x0016,          /* not referenced from symbol */
+
+  /* leaf indices starting records but referenced only from type records */
+
+  PDB_LF_SKIP_16t = 0x0200,
+  PDB_LF_ARGLIST_16t = 0x0201,
+  PDB_LF_DEFARG_16t = 0x0202,
+  PDB_LF_LIST = 0x0203,
+  PDB_LF_FIELDLIST_16t = 0x0204,
+  PDB_LF_DERIVED_16t = 0x0205,
+  PDB_LF_BITFIELD_16t = 0x0206,
+  PDB_LF_METHODLIST_16t = 0x0207,
+  PDB_LF_DIMCONU_16t = 0x0208,
+  PDB_LF_DIMCONLU_16t = 0x0209,
+  PDB_LF_DIMVARU_16t = 0x020a,
+  PDB_LF_DIMVARLU_16t = 0x020b,
+  PDB_LF_REFSYM = 0x020c,
+
+  PDB_LF_BCLASS_16t = 0x0400,
+  PDB_LF_VBCLASS_16t = 0x0401,
+  PDB_LF_IVBCLASS_16t = 0x0402,
+  PDB_LF_ENUMERATE_ST = 0x0403,
+  PDB_LF_FRIENDFCN_16t = 0x0404,
+  PDB_LF_INDEX_16t = 0x0405,
+  PDB_LF_MEMBER_16t = 0x0406,
+  PDB_LF_STMEMBER_16t = 0x0407,
+  PDB_LF_METHOD_16t = 0x0408,
+  PDB_LF_NESTTYPE_16t = 0x0409,
+  PDB_LF_VFUNCTAB_16t = 0x040a,
+  PDB_LF_FRIENDCLS_16t = 0x040b,
+  PDB_LF_ONEMETHOD_16t = 0x040c,
+  PDB_LF_VFUNCOFF_16t = 0x040d,
+
+  /* 32-bit type index versions of leaves, all have the 0x1000 bit set */
+  PDB_LF_TI16_MAX = 0x1000,
+
+  PDB_LF_MODIFIER = 0x1001,
+  PDB_LF_POINTER = 0x1002,
+  PDB_LF_ARRAY_ST = 0x1003,
+  PDB_LF_CLASS_ST = 0x1004,
+  PDB_LF_STRUCTURE_ST = 0x1005,
+  PDB_LF_UNION_ST = 0x1006,
+  PDB_LF_ENUM_ST = 0x1007,
+  PDB_LF_PROCEDURE = 0x1008,
+  PDB_LF_MFUNCTION = 0x1009,
+  PDB_LF_COBOL0 = 0x100a,
+  PDB_LF_BARRAY = 0x100b,
+  PDB_LF_DIMARRAY_ST = 0x100c,
+  PDB_LF_VFTPATH = 0x100d,
+  PDB_LF_PRECOMP_ST = 0x100e,             /* not referenced from symbol */
+  PDB_LF_OEM = 0x100f,                    /* oem definable type string */
+  PDB_LF_ALIAS_ST = 0x1010,               /* alias (typedef) type */
+  PDB_LF_OEM2 = 0x1011,                   /* oem definable type string */
+
+  /* leaf indices starting records but referenced only from type records */
+
+  PDB_LF_SKIP = 0x1200,
+  PDB_LF_ARGLIST = 0x1201,
+  PDB_LF_DEFARG_ST = 0x1202,
+  PDB_LF_FIELDLIST = 0x1203,
+  PDB_LF_DERIVED = 0x1204,
+  PDB_LF_BITFIELD = 0x1205,
+  PDB_LF_METHODLIST = 0x1206,
+  PDB_LF_DIMCONU = 0x1207,
+  PDB_LF_DIMCONLU = 0x1208,
+  PDB_LF_DIMVARU = 0x1209,
+  PDB_LF_DIMVARLU = 0x120a,
+
+  PDB_LF_BCLASS = 0x1400,
+  PDB_LF_VBCLASS = 0x1401,
+  PDB_LF_IVBCLASS = 0x1402,
+  PDB_LF_FRIENDFCN_ST = 0x1403,
+  PDB_LF_INDEX = 0x1404,
+  PDB_LF_MEMBER_ST = 0x1405,
+  PDB_LF_STMEMBER_ST = 0x1406,
+  PDB_LF_METHOD_ST = 0x1407,
+  PDB_LF_NESTTYPE_ST = 0x1408,
+  PDB_LF_VFUNCTAB = 0x1409,
+  PDB_LF_FRIENDCLS = 0x140a,
+  PDB_LF_ONEMETHOD_ST = 0x140b,
+  PDB_LF_VFUNCOFF = 0x140c,
+  PDB_LF_NESTTYPEEX_ST = 0x140d,
+  PDB_LF_MEMBERMODIFY_ST = 0x140e,
+  PDB_LF_MANAGED_ST = 0x140f,
+
+  /* Types w/ SZ names */
+
+  PDB_LF_ST_MAX = 0x1500,
+
+  PDB_LF_TYPESERVER = 0x1501,             /* not referenced from symbol */
+  PDB_LF_ENUMERATE = 0x1502,
+  PDB_LF_ARRAY = 0x1503,
+  PDB_LF_CLASS = 0x1504,
+  PDB_LF_STRUCTURE = 0x1505,
+  PDB_LF_UNION = 0x1506,
+  PDB_LF_ENUM = 0x1507,
+  PDB_LF_DIMARRAY = 0x1508,
+  PDB_LF_PRECOMP = 0x1509,                /* not referenced from symbol */
+  PDB_LF_ALIAS = 0x150a,                  /* alias (typedef) type */
+  PDB_LF_DEFARG = 0x150b,
+  PDB_LF_FRIENDFCN = 0x150c,
+  PDB_LF_MEMBER = 0x150d,
+  PDB_LF_STMEMBER = 0x150e,
+  PDB_LF_METHOD = 0x150f,
+  PDB_LF_NESTTYPE = 0x1510,
+  PDB_LF_ONEMETHOD = 0x1511,
+  PDB_LF_NESTTYPEEX = 0x1512,
+  PDB_LF_MEMBERMODIFY = 0x1513,
+  PDB_LF_MANAGED = 0x1514,
+  PDB_LF_TYPESERVER2 = 0x1515,
+
+  PDB_LF_STRIDED_ARRAY = 0x1516,       /* same as LF_ARRAY, but with stride between adjacent elements */
+  PDB_LF_HLSL = 0x1517,
+  PDB_LF_MODIFIER_EX = 0x1518,
+  PDB_LF_INTERFACE = 0x1519,
+  PDB_LF_BINTERFACE = 0x151a,
+  PDB_LF_VECTOR = 0x151b,
+  PDB_LF_MATRIX = 0x151c,
+
+  PDB_LF_VFTABLE = 0x151d,               /* a virtual function table */
+  PDB_LF_ENDOFLEAFRECORD = PDB_LF_VFTABLE,
+
+  PDB_LF_TYPE_LAST,                    /* one greater than the last type record */
+  PDB_LF_TYPE_MAX = PDB_LF_TYPE_LAST - 1,
+
+  PDB_LF_FUNC_ID = 0x1601,             /* static func ID */
+  PDB_LF_MFUNC_ID = 0x1602,            /* member func ID */
+  PDB_LF_BUILDINFO = 0x1603,           /* build info: tool, version, command line, src/pdb file */
+  PDB_LF_SUBSTR_LIST = 0x1604,         /* similar to LF_ARGLIST, for list of sub strings */
+  PDB_LF_STRING_ID = 0x1605,           /* string ID */
+
+  PDB_LF_UDT_SRC_LINE = 0x1606,        /* source and line on where an UDT is defined */
+  /* only generated by compiler */
+
+  PDB_LF_UDT_MOD_SRC_LINE = 0x1607,     /* module, source and line on where an UDT is defined */
+  /* only generated by linker */
+
+  PDB_LF_CLASSPTR = 0x1608,
+  PDB_LF_CLASSPTR2 = 0x1609,
+
+  PDB_LF_ID_LAST,                      /* one greater than the last ID record */
+  PDB_LF_ID_MAX = PDB_LF_ID_LAST - 1,
+
+  PDB_LF_NUMERIC = 0x8000,
+  PDB_LF_CHAR = 0x8000,
+  PDB_LF_SHORT = 0x8001,
+  PDB_LF_USHORT = 0x8002,
+  PDB_LF_LONG = 0x8003,
+  PDB_LF_ULONG = 0x8004,
+  PDB_LF_REAL32 = 0x8005,
+  PDB_LF_REAL64 = 0x8006,
+  PDB_LF_REAL80 = 0x8007,
+  PDB_LF_REAL128 = 0x8008,
+  PDB_LF_QUADWORD = 0x8009,
+  PDB_LF_UQUADWORD = 0x800a,
+  PDB_LF_REAL48 = 0x800b,
+  PDB_LF_COMPLEX32 = 0x800c,
+  PDB_LF_COMPLEX64 = 0x800d,
+  PDB_LF_COMPLEX80 = 0x800e,
+  PDB_LF_COMPLEX128 = 0x800f,
+  PDB_LF_VARSTRING = 0x8010,
+
+  PDB_LF_OCTWORD = 0x8017,
+  PDB_LF_UOCTWORD = 0x8018,
+
+  PDB_LF_DECIMAL = 0x8019,
+  PDB_LF_DATE = 0x801a,
+  PDB_LF_UTF8STRING = 0x801b,
+
+  PDB_LF_REAL16 = 0x801c,
+
+  PDB_LF_PAD0 = 0xf0,
+  PDB_LF_PAD1 = 0xf1,
+  PDB_LF_PAD2 = 0xf2,
+  PDB_LF_PAD3 = 0xf3,
+  PDB_LF_PAD4 = 0xf4,
+  PDB_LF_PAD5 = 0xf5,
+  PDB_LF_PAD6 = 0xf6,
+  PDB_LF_PAD7 = 0xf7,
+  PDB_LF_PAD8 = 0xf8,
+  PDB_LF_PAD9 = 0xf9,
+  PDB_LF_PAD10 = 0xfa,
+  PDB_LF_PAD11 = 0xfb,
+  PDB_LF_PAD12 = 0xfc,
+  PDB_LF_PAD13 = 0xfd,
+  PDB_LF_PAD14 = 0xfe,
+  PDB_LF_PAD15 = 0xff
+} pdb_lf_e;
+
+typedef struct pdb_lf_label
+{
+  u16 mode; // pdb_cv_ptrmode
+} pdb_lf_label;
+
+typedef struct pdb_lf_modsrcline
+{
+  pdb_cv_itype udt_itype;
+  pdb_cv_itemid src;
+  u32 ln;
+  u16 mod;
+} pdb_lf_modsrcline;
+
+typedef struct pdb_lf_srcline
+{
+  pdb_cv_itype udt_itype;
+  pdb_cv_itemid src;
+  u32 ln;
+} pdb_lf_srcline;
+
+typedef struct pdb_lf_func_id 
+{
+  pdb_cv_itemid scope_id;    /* parent scope of the ID, 0 if global */
+  pdb_cv_itype itype;           /* function type */
+  /* unsigned char   name[CV_ZEROLEN]; */
+} pdb_lf_func_id;
+
+typedef struct pdb_lf_mfunc_id
+{
+  pdb_cv_itemid parent_itype;
+  pdb_cv_itemid itype;
+  /* unsigned char name[CV_ZEROLEN]; */
+} pdb_lf_mfunc_id;
+
+typedef struct pdb_lf_string_id
+{
+  pdb_cv_itemid id;
+  /* unsigned char name[CV_ZEROLEN] */
+} pdb_lf_string_id;
+
+typedef struct pdb_lf_mfunc 
+{
+  /* type index of return value */
+  pdb_cv_itype rvtype;
+
+  /* type index of containing class */
+  pdb_cv_itype classtype;
+
+  /* type index of this pointer (model specific) */
+  pdb_cv_itype thistype;
+
+  /* calling convention (call_t) */
+  U8 calltype;
+
+  /* attributes */
+  pdb_cv_funcattr_t funcattr;
+
+  /* number of parameters */
+  U16 parmcount;
+
+  /* type index of argument list */
+  pdb_cv_itype arglist;
+
+  /* this adjuster (long because pad required anyway) */
+  U32 thisadjust;     
+} pdb_lf_mfunc;
+
+typedef struct pdb_lf_modifier
+{
+  /* type index of modified type */
+  pdb_cv_itype itype;
+  pdb_cv_modifier_t attr;
+} pdb_lf_modifier;
+
+typedef struct pdb_lf_class 
+{
+  /* count of number of elements in class */
+  U16 count;               
+
+  /* (pdb_cv_prop_t) property attribute field */
+  pdb_cv_prop_t prop;          
+
+  /* (type index) type index of LF_FIELD descriptor list */
+  U32 field;               
+
+  /* (type index) type index of derived from list if not zero */
+  U32 derived;             
+
+  /* (type index) type index of vshape table for this class */
+  U32 vshape;              
+
+  /* char data[] size and name follow */
+} pdb_lf_class;
+
+typedef struct pdb_lf_enum 
+{
+  /* number of elements */
+  U16 count;             
+
+  /* property attribute field */
+  pdb_cv_prop_t prop;        
+
+  /* underlying type of the enum */
+  pdb_cv_itype itype;             
+
+  /* type index of LF_FIELD descriptor list */
+  U32 field;             
+
+  /* unsigned char name[1];  length prefixed name of enum */
+} pdb_lf_enum;
+
+typedef struct pdb_lf_onemethod 
+{
+  pdb_cv_fldattr_t attr;
+  pdb_cv_itype itype;
+  /* U32 vbaseoff[]; */
+} pdb_lf_onemethod;
+
+typedef struct pdb_lf_stmember 
+{
+  /* attribute mask */
+  pdb_cv_fldattr_t attr;
+  
+  /* index of type record for field */
+  pdb_cv_itype index;          
+
+  /* char name[]; */
+} pdb_lf_stmember;
+
+typedef struct pdb_lf_method 
+{
+  U16 count;
+  pdb_cv_itype itype_list;
+  /* char name[] */
+} pdb_lf_method;
+
+typedef struct pdb_lf_blcass 
+{
+  pdb_cv_fldattr_t attr;
+  pdb_cv_itype itype;
+  /* U32 offset[] */
+} pdb_lf_bclass;
+
+typedef struct pdb_lf_vfunctab 
+{
+  char pad[2];
+  pdb_cv_itype itype;
+} pdb_lf_vfunctab;
+
+typedef struct pdb_lf_vftable
+{
+  pdb_cv_itype owner_itype;
+  pdb_cv_itype base_table_itype;
+  u32 offset_in_object_layout;
+  u32 names_len;
+  /*char names[0];*/ // Name of the table comes first and then follow method names
+} pdb_lf_vftable;
+
+typedef struct pdb_lf_vfuncoff
+{
+  char pad[2];
+  pdb_cv_itype itype;
+  u32 offset;
+} pdb_lf_vfuncoff;
+
+typedef struct pdb_lf_vbclass
+{
+  pdb_cv_fldattr_t attr;
+  pdb_cv_itype index;
+  pdb_cv_itype vbptr;
+  // u8 vbpoff[0];
+  // u32 virtual_base_offset_from_vtable;
+} pdb_lf_vbclass;
+
+typedef struct pdb_lf_nesttype 
+{
+  char pad[2];
+  pdb_cv_itype itype;
+  /* char name[]; */
+} pdb_lf_nesttype;
+
+typedef struct pdb_off_cb 
+{
+  U32 off;
+  U32 cb;
+} pdb_off_cb;
+
+typedef struct pdb_lf_index 
+{
+  char pad[2];
+
+  /* type index of refed leaf */
+  pdb_cv_itype itype; 
+} pdb_lf_index;
+
+typedef struct pdb_lf_array 
+{
+  /* type index of element type */
+  pdb_cv_itype entry_itype;              
+  
+  /* type index of indexing type */
+  pdb_cv_itype index_itype;               
+
+  /* char data[]; variable length data specifying size in bytes and name */
+} pdb_lf_array;
+
+typedef struct pdb_lf_enumerate 
+{
+  pdb_cv_fldattr_t attr;        
+  /* char value[0]; variable length value field followed by length prefixed name */
+} pdb_lf_enumerate;
+
+typedef struct pdb_lf_member 
+{
+  pdb_cv_fldattr_t attr;
+
+  pdb_cv_itype itype;
+
+  /* U8 offset[] // variable length offset of field followed by length prefixed name of field */
+} pdb_lf_member;
+
+typedef struct pdb_lf_union 
+{
+  /* count of number of elements in class */
+  U16 count;               
+
+  /* property attribute field */
+  pdb_cv_prop_t prop;          
+
+  /* type index of LF_FIELD descriptor list a numeric and name follow. */
+  U32 field;               
+} pdb_lf_union;
+
+typedef struct pdb_lf_arglist 
+{
+  U32 count;
+  /* U32 *itypes; */
+} pdb_lf_arglist;
+
+typedef struct pdb_lf_proc 
+{
+  /* type index of return value */
+  pdb_cv_itype ret_itype;              
+  
+  /* calling convention (pdb_cv_call_e) */
+  U8 call_kind;              
+
+  pdb_cv_funcattr_t funcattr; 
+
+  /* number of parameters */
+  U16 arg_count;              
+
+  /* type index of argument list */
+  pdb_cv_itype arg_itype;              
+} pdb_lf_proc;
+
+typedef struct pdb_lf_mproc 
+{
+  /* type index of return value */
+  pdb_cv_itype ret_itype;      
+
+  /* type index of containing class */
+  pdb_cv_itype class_itype;    
+
+  /* type index of this pointer (model specific) */
+  pdb_cv_itype this_itype;     
+  
+  /* pdb_cv_call_e */
+  U8 call_kind;      
+
+  pdb_cv_funcattr_t funcattr;       
+
+  /* number of parameters */
+  U16 arg_count;      
+
+  /* type index of argument list */
+  pdb_cv_itype arg_itype;      
+
+  /* this adjuster (s32 because pad required anyway) */
+  S32 thisadjust;     
+} pdb_lf_mproc;
+
+typedef struct pdb_lf_ptr 
+{ 
+  pdb_cv_itype itype;   
+  pdb_cv_ptr_attrib_t attr;
+} pdb_lf_ptr;
+
+typedef struct pdb_lf_bitfield 
+{
+  pdb_cv_itype itype;
+  U8 len;
+  U8 pos;
+} pdb_lf_bitfield;
+
+//
+// Virtual Shape Table
+//
+
+enum
+{
+  PDB_CV_VTS_NEAR   = 0, // 16 bit pointer
+  PDB_CV_VTS_FAR    = 1, // 16:16 bit pointer
+  PDB_CV_VTS_THIN   = 2, // ??
+  // From MS pdf (page 36):
+  // Address point displacement to outermost class. This is at entry[-1] from table entry
+  PDB_CV_VTS_OUTER  = 3,
+  // From MS pdf (page 36):
+  // Far pointer to metaclass descriptor. This is at entry[-2] from table address.
+  PDB_CV_VTS_META   = 4,
+  PDB_CV_VTS_NEAR32 = 5, // 32bit pointer
+  PDB_CV_VTS_FAR32  = 6, // ??
+  PDB_CV_VTS_UNUSED = 7
+};
+typedef u32 pdb_lf_vts_desc;
+
+typedef struct pdb_lf_vtshape
+{
+  u16 count;
+  // pdb_lf_vts_desc desc[0];
+} pdb_lf_vtshape;
+
+typedef struct pdb_ml_method
+{
+  pdb_cv_fldattr_t attr;
+  u8 pad[2];
+  pdb_cv_itype index;
+  // unsigned long vbaseoff[0];
+} pdb_ml_method;
+
+typedef struct pdb_ml_methodlist
+{
+  pdb_ml_method list[1];
+} pdb_ml_methodlist;
+
+//
+// 
+//
+
+typedef struct pdb_lf_classptr
+{
+  u32 prop; // pdb_cv_prop_t but 32bit
+  pdb_cv_itype arglist_itype;
+  u32 unknown2;
+  u32 unknown3;
+  u16 unknown4;
+  //char size[];
+  //char name[0];
+  //char linkage_name[0];
+} pdb_lf_classptr;
+
+#pragma pack(pop) /* CodeView bottom */
+
+#endif /* SYMS_CODEVIEW_INCLUDE_H */
diff --git a/Engine/Source/Runtime/SymsLib/include/pdb/syms_dbi.h b/Engine/Source/Runtime/SymsLib/include/pdb/syms_dbi.h
new file mode 100644
--- /dev/null
+++ b/Engine/Source/Runtime/SymsLib/include/pdb/syms_dbi.h
@@ -0,0 +1,337 @@
+// Copyright Epic Games, Inc. All Rights Reserved.
+#ifndef SYMS_DBI_INCLUDE_H
+#define SYMS_DBI_INCLUDE_H
+
+typedef enum 
+{
+  PDB_DBI_VER_41  = 930803,
+  PDB_DBI_VER_50  = 19960307,
+  PDB_DBI_VER_60  = 19970606,
+  PDB_DBI_VER_70  = 19990903,
+  PDB_DBI_VER_110 = 20091201,
+
+  PDB_DBI_VER_CUR = PDB_DBI_VER_70
+} pdb_dbi_header_ver_e;
+
+#define PDB_DBI_SC_VER_60 (0xeffe0000ul + 19970605ul)
+#define PDB_DBI_SC_VER_2  (0xeffe0000ul + 20140516ul)
+typedef U32 pdb_dbi_sc_ver;
+
+#define PDB_DBI_HEADER_SIG_NULL 0
+#define PDB_DBI_HEADER_SIG_V1   0xFFFFFFFF
+typedef U32 pdb_dbi_header_sig;
+
+typedef U32 pdb_ni_t;
+
+/* name index */
+typedef U32 pdb_ni;
+
+typedef enum 
+{
+  PDB_DBG_STREAM_FPO,
+  PDB_DBG_STREAM_EXCEPTION,
+  PDB_DBG_STREAM_FIXUP,
+  PDB_DBG_STREAM_OMAP_TO_SRC,
+  PDB_DBG_STREAM_OMAP_FROM_SRC,
+  PDB_DBG_STREAM_SECTION_HEADER, // Stream contains array of Image_Section_Header.
+  PDB_DBG_STREAM_TOKEN_RDI_MAP,
+  PDB_DBG_STREAM_XDATA,
+  PDB_DBG_STREAM_PDATA,
+  PDB_DBG_STREAM_NEW_FPO,
+  PDB_DBG_STREAM_SECTION_HEADER_ORIG,
+
+  PDB_DBG_STREAM_MAX
+} pdb_dbg_stream_e;
+
+typedef struct pdb_sc40 
+{
+  pdb_isec sec;
+  U32 sec_off;
+  U32 size;
+  U32 flags;
+  pdb_imod imod;
+
+  U8 r[7];
+} pdb_sc40;
+
+typedef struct pdb_sc20 
+{
+  pdb_isec sec;
+  U32 sec_off;
+  U32 size;
+  pdb_imod imod;
+} pdb_sc20;
+
+typedef struct pdb_sc 
+{
+  pdb_isec sec;
+  U8 padding1[2];
+  pdb_isec_umm sec_off;
+  pdb_isec_umm size;
+  U32 flags;
+  pdb_imod imod;
+  U8 padding2[2];
+  U32 data_crc;
+  U32 reloc_crc;
+} pdb_sc;
+
+typedef struct pdb_sc2 
+{
+  pdb_isec sec;
+  U8 padding1[2];
+  pdb_isec_umm sec_off;
+  pdb_isec_umm size;
+  U32 flags;
+  pdb_imod imod;
+  U8 padding2[2];
+  U32 data_crc;
+  U32 reloc_crc;
+  U32 sec_coff;
+} pdb_sc2;
+
+#pragma pack(push, 1)
+
+enum
+{
+  PDB_DBI_HEADER_BUILD_MINOR_MASK = 0x00FF,
+  PDB_DBI_HEADER_BUILD_MINOR_SHIFT = 0,
+
+  PDB_DBI_HEADER_BUILD_MAJOR_MASK = 0x7F00,
+  PDB_DBI_HEADER_BUILD_MAJOR_SHIFT = 8,  
+
+  PDB_DBI_HEADER_NEW_VERSION_FORMAT_MASK = 0x8000
+};
+
+enum
+{
+  PDB_DBI_HEADER_FLAGS_INCREMENTAL_MASK = 0x1, // bit is set if linked incrementally
+  PDB_DBI_HEADER_FLAGS_STRIPPED_MASK    = 0x2, // bit is set if private symbols were stripped
+  PDB_DBI_HEADER_FLAGS_CTYPES_MASK      = 0x4  // bit is set if linked with /debug:ctypes
+};
+typedef u16 pdb_dbi_header_flags;
+
+typedef struct pdb_dbi_header 
+{
+  pdb_dbi_header_sig sig;
+  U32 version;
+  U32 age;
+  pdb_sn global_sym_sn;
+
+  U16 build_number;
+
+  /* public symbol stream. */
+  pdb_sn public_sym_sn;  
+
+  /* build version of the pdb dll that built this pdb last. */
+  U16 pdb_version;                
+
+  /* An array of symbols, nothing special (Format u16 size, u16 type, char data[size]) */
+  pdb_sn sym_record_sn;  
+  
+  /* rbld version of the pdb dll that built this pdb last. */
+  U16 pdb_version2;               
+
+  U32 module_info_size;
+  U32 sec_con_size;
+  U32 sec_map_size;
+  U32 file_info_size;
+
+  /* size of the Type Server Map substream. */
+  U32 tsm_size;           
+
+  /* index of MFC type server */
+  U32 mfc_index;          
+
+  /* size of optional DbgHdr info appended to the end of the stream */
+  U32 dbg_header_size;    
+  
+  /* number of bytes in EC substream, or 0 if EC no EC enabled Mods */
+  U32 ec_info_size;       
+
+  U16 flags;
+  
+  /* ImageFileMachine */
+  U16 machine;     
+
+  U32 reserved;
+} pdb_dbi_header;
+
+typedef struct pdb_mod_header 
+{
+  U8 unused[4]; 
+  
+  pdb_sc sc;
+
+  U16 flags;
+
+  /* stream number of module debug info (syms, lines, FPO) */
+  pdb_sn sn;    
+
+  /* size of local symbol debug info in above stream */
+  U32 symbol_bytes;      
+
+  /* size of line number debug info in stream sn */
+  U32 c11_lines_size;    
+
+  /* size of C13 style line number info in stream sn */
+  U32 c13_lines_size;    
+
+  /* number of files contributing to this module (file index) */
+  U16 num_contrib_files; 
+
+  U8 padding1[2];
+
+  /* Shit field that is unused in our case */
+  U32 file_names_offset; 
+
+  pdb_ni src_file_ni;
+  pdb_ni pdb_file_ni;
+
+  /* char module_name[]; */
+  /* char obj_name[]; */
+} pdb_mod_header;
+
+typedef struct pdb_psi_header 
+{
+  U32 sym_hash_size;
+  U32 addr_map_size;
+  U32 thunk_count;
+  U32 thunk_size; 
+  pdb_isec isec_thunk_table;
+  U8 padding[2];
+  pdb_isec_umm sec_thunk_tabl_off; 
+  U32 sec_count;
+} pdb_psi_header;
+
+#define PDB_GSI_V70   0xeffe0000 + 19990810
+#define PDB_GSI_SIG   (~0u)
+typedef struct pdb_gsi_header 
+{
+  U32 sig;
+  U32 ver;
+  U32 hr_len;
+  U32 num_buckets;
+} pdb_gsi_header;
+
+typedef struct pdb_gsi_file_hr 
+{
+  U32 off; /* Offset in the symbol record stream */
+  U32 cref;
+} pdb_gsi_file_hr;
+
+typedef struct pdb_gsi_hr 
+{
+  U32 off;
+  struct pdb_gsi_hr *next;
+} pdb_gsi_hr;
+
+typedef struct pdb_ti_off 
+{
+  pdb_ti ti;
+  U32 off;
+} pdb_ti_off;
+
+#define PDB_INTV_VC2 0xE0D5C
+
+#define PDB_TM_IMPV40          0x1306B4A
+#define PDB_TM_IMPV41          0x1306E12
+#define PDB_TM_IMPV50_INTERIM  0x13091F3
+#define PDB_TM_IMPV50          0x13094C7
+#define PDB_TM_IMPV70          0x1310977
+#define PDB_TM_IMPV80          0x131CA0B
+#define PDB_TM_IMV_CURR        PDB_TM_IMV80
+typedef struct pdb_tm_header 
+{
+  U32 version;
+
+  /* Size of this header */
+  U32 header_size;
+
+  /* All types below the "ti_lo" are reserved for internal use. Most of them
+   * represent basic types, like int, float, char and so on. */
+  U32 ti_lo;
+  U32 ti_hi;
+
+  /* Size in bytes of all types that follow pdb_tm_header. */
+  U32 types_size;
+
+  /* Hash stream is sub-divided into following sections:
+   *
+   * pdb_ti hashes[ti_hi - ti_lo] 
+   *
+   * pdb_ti_off offsets[ti_off.cb/sizeof(pdb_ti_off)] 
+   *
+   * Last section is a serialized map of types that have to be moved in the
+   * internally-linked table (pdb_tm.buckets). Size of this sub-section
+   * is defined by the "hash_adj" field below. */
+  pdb_sn hash_sn;
+
+  /* If stream not present this value is 0xffff */
+  pdb_sn hash_sn_aux; 
+
+  /* Size of the hash value in the "hash_sn". */
+  S32 hash_key_size;
+
+  /* Count of hashes that map string name of 
+   * a type to it's corresponding type index. */
+  U32 hash_bucket_count;
+
+  pdb_off_cb hash_vals;
+  pdb_off_cb ti_off;
+  pdb_off_cb hash_adj;
+} pdb_tm_header;
+
+typedef struct pdb_img_sec 
+{
+  char name[8];
+
+  union {
+    U32 physical_address;
+    U32 virtual_size;
+  } misc;
+
+  U32 rva;
+  U32 sizeof_raw_data;
+  U32 pointer_to_raw_data;
+  U32 pointer_to_relocations;
+  U32 pointer_to_linenumbers;
+  U16 relocations_count;
+  U16 linenumbers_count;
+  U32 characteristics;
+} pdb_img_sec;
+
+enum
+{
+  PDB_OMF_SEG_DESC_FLAGS_NULL          = 0,
+  PDB_OMF_SEG_DESC_FLAGS_READ          = (1u << 0u),
+  PDB_OMF_SEG_DESC_FLAGS_WRITE         = (1u << 1u),
+  PDB_OMF_SEG_DESC_FLAGS_EXEC          = (1u << 2u),
+  PDB_OMF_SEG_DESC_FLAGS_ADDR_IS_32BIT = (1u << 3u),
+  PDB_OMF_SEG_DESC_FLAGS_IS_SELECTOR   = (1u << 8u),
+  PDB_OMF_SEG_DESC_FLAGS_IS_ABS_ADDR   = (1u << 9u),
+  PDB_OMF_SEG_DESC_FLAGS_IS_GROUP      = (1u << 10u)
+};
+typedef u16 pdb_omf_flags;
+
+typedef struct pdb_secmap_header
+{
+  u16 sec_count;      // Count of segment descriptors in the table
+  u16 sec_count_log;  // Count of logical segments in the table
+} pdb_secmap_header;
+
+typedef struct pdb_secmap_entry
+{
+  u16 flags;           // Entry flags. See pdb_omf_flags
+  u16 ovl;             // Logical overlay number
+  u16 group;           // Index of group in the descriptor array
+  u16 frame;           // ??
+  u16 sec_name;        // Index in the sstSegName table, or if not present SYMS_UINT16_MAX
+  u16 class_name;      // Index in the sstSegName table, or if not present SYMS_UINT16_MAX
+  u32 offset;          // Logical segment offset within the physical segment
+  u32 sec_byte_length; // Number of bytes that make up group or segment
+} pdb_secmap_entry;
+
+#pragma pack(pop)
+
+#endif /* SYMS_DBI_INCLUDE_H */
+
diff --git a/Engine/Source/Runtime/SymsLib/include/pdb/syms_msf.h b/Engine/Source/Runtime/SymsLib/include/pdb/syms_msf.h
new file mode 100644
--- /dev/null
+++ b/Engine/Source/Runtime/SymsLib/include/pdb/syms_msf.h
@@ -0,0 +1,161 @@
+// Copyright Epic Games, Inc. All Rights Reserved.
+#ifndef SYMS_MSF_INCLUDE_H
+#define SYMS_MSF_INCLUDE_H
+
+#ifndef PDB_API
+#define PDB_API
+#endif
+
+typedef U16 pdb_uint16;
+typedef S32 pdb_int;
+typedef U32 pdb_uint;
+typedef U32 pdb_umm;
+#define PDB_UINT_MAX SYMS_UINT32_MAX
+#define pdb_trunc_uint(x) (pdb_uint)syms_trunc_u32(x)
+#define pdb_trunc_int(x)  (pdb_int)syms_trunc_u32(x)
+
+struct pdb_context;
+
+typedef u32 pdb_offset;
+
+/* stream number */
+typedef u16 pdb_sn;
+
+#define PDB_INVALID_SN 0xffff
+
+typedef struct pdb_header20 
+{
+  U8 magic[44];
+  U32 page_size;
+  U16 free_page_map;
+  U16 pages_used;
+  U32 root_size;
+  U32 reserved;
+  U16 page_map_addr;
+} pdb_header20;
+
+typedef struct pdb_header70 
+{
+  U8 magic[32];
+
+  /* Size of page in bytes */
+  U32 page_size;
+
+  /* Index of the free page map */
+  S32 free_page_map;
+
+  /* Number of pages used */
+  S32 pages_used;
+
+  /* Size of root directory in bytes */
+  U32 root_size;
+
+  /* Purpose of this field is unknown */
+  U32 reserved;
+
+  /* Number of page that contains page map */
+  U32 page_map_addr;
+} pdb_header70;
+
+typedef enum 
+{
+  PDB_PAGE_INDEX16 = 2,
+  PDB_PAGE_INDEX32 = 4
+} pdb_page_index_type;
+
+#define PDB_STREAM_FLAGS_READ_FAILED (1u << 0u)
+#define PDB_STREAM_FLAGS_SEEK_FAILED (1u << 1u)
+typedef U32 pdb_stream_flags;
+
+#define PDB_STREAM_READ_OR_SEEK_FAILED(flags) (((flags) & PDB_STREAM_FLAGS_READ_FAILED) || ((flags) & PDB_STREAM_FLAGS_SEEK_FAILED))
+#define PDB_STREAM_NO_ERROR(flags)            ((~(flags) & PDB_STREAM_FLAGS_READ_FAILED) && (~(flags) & PDB_STREAM_FLAGS_SEEK_FAILED))
+typedef struct pdb_stream 
+{
+  struct pdb_context *pdb;
+
+  pdb_stream_flags flags;
+
+  /* Stream number (ID) */
+  pdb_sn sn;
+
+  U32 page_size;
+  U32 page_read_lo;
+  U32 page_read_hi;
+
+  /* Indicates how many bytes read into the stream. */
+  U32 off;
+
+  /* Number of bytes read when stream was subset */
+  U32 off_at_subset;
+
+  /* Amount of bytes that can be read from this stream. */
+  U32 size;
+
+  /* Offset into the root stream where the page indices for this stream start. */
+  U32 indices_off;
+
+  /* Current location in the root stream. */
+  U32 root_off;
+} pdb_stream;
+
+typedef struct pdb_symrec 
+{
+  U16 size;
+  U16 type;
+  U32 end;
+} pdb_symrec;
+
+#define pdb_stream_skip(stream, num) pdb_stream_seek(stream, (stream)->off + (num))
+PDB_API syms_bool
+pdb_stream_seek(pdb_stream *stream, U32 offset);
+
+#define pdb_stream_read_struct(stream, bf, type) pdb_stream_read(stream, bf, sizeof(type))
+PDB_API pdb_uint
+pdb_stream_read(pdb_stream *stream, void *buff, pdb_uint buff_max);
+
+PDB_API syms_bool
+pdb_stream_can_read_bytes(pdb_stream *stream, pdb_uint num_bytes);
+
+PDB_API syms_bool
+pdb_stream_align(pdb_stream *stream, pdb_uint align);
+
+PDB_API syms_bool
+pdb_stream_init(struct pdb_context *pdb, pdb_sn sn, pdb_stream *out_stream);
+
+PDB_API syms_bool
+pdb_stream_init_at(struct pdb_context *pdb, pdb_sn sn, pdb_offset offset, pdb_stream *out_stream);
+
+PDB_API void
+pdb_stream_init_null(pdb_stream *stream);
+
+PDB_API syms_bool
+pdb_stream_is_null(pdb_stream *stream);
+
+PDB_API pdb_stream
+pdb_stream_subset(const pdb_stream *stream, pdb_uint off, pdb_uint size);
+
+PDB_API pdb_uint
+pdb_stream_get_abs_off(pdb_stream *stream);
+
+PDB_API syms_bool
+pdb_stream_read_u08(pdb_stream *stream, U8 *out_value);
+
+PDB_API syms_bool
+pdb_stream_read_u16(pdb_stream *stream, U16 *out_value);
+
+PDB_API syms_bool
+pdb_stream_read_u32(pdb_stream *stream, pdb_uint *out_value);
+
+PDB_API syms_bool
+pdb_stream_read_uleb32(pdb_stream *stream, pdb_uint *out_value);
+
+PDB_API syms_bool
+pdb_stream_read_symrec(pdb_stream *stream, pdb_symrec *out_rec);
+
+PDB_API pdb_uint
+pdb_stream_read_str(pdb_stream *stream, void *bf, pdb_uint bf_max);
+
+PDB_API pdb_uint
+pdb_stream_strlen(pdb_stream *stream);
+
+#endif /* SYMS_MSF_INCLUDE_H */
