RAD symbol library

diff --git a/Engine/Source/Runtime/SymsLib/include/pdb/syms_pdb.h b/Engine/Source/Runtime/SymsLib/include/pdb/syms_pdb.h
new file mode 100644
--- /dev/null
+++ b/Engine/Source/Runtime/SymsLib/include/pdb/syms_pdb.h
@@ -0,0 +1,1445 @@
+// Copyright Epic Games, Inc. All Rights Reserved.
+#ifndef SYMS_PDB_INCLUDE_H
+#define SYMS_PDB_INCLUDE_H
+
+#define SYMS_ASSERT_CORRUPTED_STREAM SYMS_ASSERT_FAILURE("corrupted stream")
+
+/* Maximum number of bytes a string contains */
+#define PDB_STRLEN_MAX 4096
+
+enum
+{
+  PDB_ERROR_ok,
+  PDB_ERROR_corrupted_data
+};
+typedef u32 pdb_error_code;
+
+typedef enum 
+{
+  PDB_VER_VC2      = 0x13048EA,
+  PDB_VER_VC4      = 0x1306C1F,
+  PDB_VER_VC41     = 0x1306CDE,
+  PDB_VER_VC50     = 0x13091F3,
+  PDB_VER_VC98     = 0x130BA2C,
+  PDB_VER_VC70     = 0x1312E94,
+  PDB_VER_VC70_DEP = 0x131084C,
+  PDB_VER_VC80     = 0x131A5B5,
+  PDB_VER_VC110    = 0x1329141,
+  PDB_VER_VC140    = 0x13351DC
+} pdb_ver_e;
+
+typedef enum 
+{
+  PDB_TYPED_STREAM_HEADER_BLOCK,
+
+  /* Contains hash tables for strings and offsets. 
+   * For parsing look at "pdb_init". */
+  PDB_TYPED_STREAM_STRTABLE,
+
+  /* Contains info left there by linker. */
+  PDB_TYPED_STREAM_LINK_INFO,
+
+  PDB_TYPED_STREAM_MAX
+} pdb_typed_stream_e;
+
+typedef enum 
+{
+  PDB_DEFAULT_STREAM_ROOT,
+  PDB_DEFAULT_STREAM_STREAMS,
+  PDB_DEFAULT_STREAM_TPI,
+  PDB_DEFAULT_STREAM_DBI,
+  PDB_DEFAULT_STREAM_IPI
+} pdb_default_stream_e;
+typedef pdb_sn pdb_default_stream;
+
+typedef struct
+{
+  u32 count;
+  pdb_cv_itype *e;
+} pdb_itype_array;
+
+#define PDB_NUMERIC_MAX 32
+typedef struct pdb_numeric 
+{
+  pdb_cv_itype itype;
+  union 
+  {
+    char ch;
+
+    S8 int8;
+    S16 int16;
+    S32 int32;
+    S64 int64;
+
+    U8 uint8;
+    U16 uint16;
+    U32 uint32;
+    U64 uint64;
+
+    F32 float32;
+    F64 float64;
+
+    U8 data[PDB_NUMERIC_MAX];
+  } u;
+} pdb_numeric;
+
+enum 
+{
+  PDB_POINTER_MODE_NULL,
+  PDB_POINTER_MODE_RAW,
+  PDB_POINTER_MODE_PAGES,
+  PDB_POINTER_MODE_STREAM
+};
+typedef u32 pdb_pointer_mode;
+
+// Maximum number of pages we can bake in one pointer.
+// Imperically largerst thing that gets baked is a string,
+// which are around 4k long and by default page size in PDB
+// is 4K, so edge case scenario here is a string that spans 
+// accross two pages.
+#define PDB_POINTER_PAGE_MAX 2
+
+typedef struct pdb_pointer 
+{
+  pdb_pointer_mode mode;
+
+  union {
+    struct {
+      const void *data;
+      pdb_uint    size;
+    } raw;
+
+    struct {
+      pdb_uint offs[PDB_POINTER_PAGE_MAX];
+      pdb_uint size[PDB_POINTER_PAGE_MAX];
+    } pages;
+  
+    struct {
+      pdb_sn   sn;
+      pdb_uint off;
+      pdb_uint size;
+    } stream;
+  } u;
+} pdb_pointer;
+typedef pdb_pointer pdb_string_ref;
+
+PDB_API pdb_pointer
+pdb_pointer_bake(pdb_stream *stream, pdb_uint num_bytes);
+
+PDB_API pdb_pointer
+pdb_pointer_bake_null(void);
+
+PDB_API pdb_pointer
+pdb_pointer_bake_one_page(void *data, U32 data_len);
+
+PDB_API pdb_pointer
+pdb_pointer_bake_str(SymsString str);
+
+PDB_API pdb_pointer
+pdb_pointer_bake_stream_str(pdb_stream *msf);
+
+PDB_API pdb_pointer
+pdb_pointer_bake_sn(struct pdb_context *pdb, pdb_sn sn, U32 off, U32 len);
+
+PDB_API U32
+pdb_pointer_get_size(pdb_pointer *pointer);
+
+PDB_API U32
+pdb_pointer_read(struct pdb_context *pdb, const pdb_pointer *pointer, U32 off, void *bf, U32 bf_max);
+
+PDB_API U32
+pdb_pointer_read_u32(struct pdb_context *pdb, pdb_pointer *pointer, U32 off);
+
+PDB_API U16
+pdb_pointer_read_u16(struct pdb_context *pdb, pdb_pointer *pointer, U32 off);
+
+PDB_API U8
+pdb_pointer_read_u08(struct pdb_context *pdb, pdb_pointer *pointer, U32 off);
+
+PDB_API U32
+pdb_pointer_read_utf8(struct pdb_context *pdb, pdb_pointer *pointer, U32 off, U32 *codepoint_out);
+
+PDB_API syms_bool
+pdb_pointer_cmp(struct pdb_context *pdb, pdb_pointer *pointer_a, pdb_pointer *pointer_b);
+
+typedef enum 
+{
+  PDB_STRCMP_FLAG_NULL     = 0,
+  PDB_STRCMP_FLAG_NOCASE   = (1 << 0),
+  PDB_STRCMP_FLAG_NOREWIND = (1 << 1)
+} pdb_strcmp_flags_e;
+
+PDB_API syms_bool
+pdb_pointer_strcmp_(struct pdb_context *pdb, pdb_pointer *pointer_a, pdb_pointer *pointer, pdb_strcmp_flags_e cmp_flags);
+
+#define pdb_strcmp_pointer(pdb, a, b)        pdb_strcmp_pointer_(pdb, a, b, PDB_STRCMP_FLAG_NULL)
+#define pdb_strcmp_pointer_nocase(pdb, a, b) pdb_strcmp_pointer_(pdb, a, b, PDB_STRCMP_FLAG_NOCASE)
+PDB_API syms_bool
+pdb_strcmp_pointer_(struct pdb_context *pdb, SymsString str, pdb_pointer *str_b, pdb_strcmp_flags_e cmp_flags);
+
+#define pdb_pointer_strcmp(pdb, a, b)        pdb_pointer_strcmp_(pdb, a, b, PDB_STRCMP_FLAG_NULL)
+#define pdb_pointer_strcmp_nocase(pdb, a, b) pdb_pointer_strcmp_(pdb, a, b, PDB_STRCMP_FLAG_NOCASE)
+PDB_API syms_bool
+pdb_pointer_strcmp_(struct pdb_context *pdb, pdb_pointer *str_a, pdb_pointer *str_b, pdb_strcmp_flags_e cmp_flags);
+
+#define pdb_strcmp_stream(a, b)        pdb_strcmp_stream_(a, b, PDB_STRCMP_FLAG_NULL)
+#define pdb_strcmp_stream_nocase(a, b) pdb_strcmp_stream_(a, b, PDB_STRCMP_FLAG_NOCASE)
+PDB_API syms_bool
+db_strcmp_stream_(SymsString str, struct pdb_stream *stream, pdb_strcmp_flags_e cmp_flags);
+
+#define pdb_stream_strcmp(a, b)        pdb_stream_strcmp_stream_(a, b, PDB_STRCMP_FLAG_NULL)
+#define pdb_stream_strcmp_nocase(a, b) pdb_stream_strcmp_stream_(a, b, PDB_STRCMP_FLAG_NOCASE)
+PDB_API syms_bool
+pdb_stream_strcmp_stream_(pdb_stream *stream_a, pdb_stream *stream_b, pdb_strcmp_flags_e cmp_flags);
+
+#define pdb_stream_strcmp_pointer(a, b)         pdb_stream_strcmp_pointer_(a, b, PDB_STRCMP_FLAG_NULL)
+#define pdb_stream_strcmp_pointer_nocase(a, b)  pdb_stream_strcmp_pointer_(a, b, PDB_STRCMP_FLAG_NOCASE)
+PDB_API syms_bool
+pdb_stream_strcmp_pointer_(pdb_stream *str_a, pdb_pointer *str_b, pdb_strcmp_flags_e cmp_flags);
+
+typedef enum 
+{
+  PDB_TYPE_NULL,
+  PDB_TYPE_CHAR,
+  PDB_TYPE_UCHAR,
+  PDB_TYPE_INT8,
+  PDB_TYPE_UINT8,
+  PDB_TYPE_INT16,
+  PDB_TYPE_UINT16,
+  PDB_TYPE_INT32,
+  PDB_TYPE_UINT32,
+  PDB_TYPE_INT64,
+  PDB_TYPE_UINT64,
+  PDB_TYPE_INT128,
+  PDB_TYPE_UINT128,
+  PDB_TYPE_REAL16,
+  PDB_TYPE_REAL32,
+  PDB_TYPE_REAL32PP, /* 32 bit partial precision */
+  PDB_TYPE_REAL64,
+  PDB_TYPE_REAL80,
+  PDB_TYPE_REAL128,
+  PDB_TYPE_BOOL,
+  PDB_TYPE_VOID,
+  PDB_TYPE_WCHAR,
+  PDB_TYPE_STRUCT,
+  PDB_TYPE_UNION,
+  PDB_TYPE_CLASS,
+  PDB_TYPE_ENUM,
+  PDB_TYPE_PROC,
+  PDB_TYPE_ENUM_FIELD,
+  PDB_TYPE_BITFIELD,
+  PDB_TYPE_NESTED_UDT,
+  PDB_TYPE_FWDREF,
+  PDB_TYPE_VTSHAPE,
+  PDB_TYPE_METHOD,
+  PDB_TYPE_METHODLIST,
+  PDB_TYPE_ARGLIST,
+  PDB_TYPE_VARIADIC,
+  PDB_TYPE_POINTER,
+  PDB_TYPE_PTR,
+  PDB_TYPE_ARR,
+  PDB_TYPE_FUNCID,
+  PDB_TYPE_MFUNCID,
+  PDB_TYPE_STRINGID,
+  PDB_TYPE_COMPLEX32,
+  PDB_TYPE_COMPLEX64,
+  PDB_TYPE_COMPLEX80,
+  PDB_TYPE_COMPLEX128,
+  PDB_TYPE_FIELDLIST,
+  PDB_TYPE_LABEL
+} pdb_type_kind_e;
+typedef U32 pdb_type_kind;
+
+typedef enum 
+{
+  PDB_TYPE_ATTRIB_HAS_NAME    = 1 << 0,
+  PDB_TYPE_ATTRIB_CONST       = 1 << 1,
+  PDB_TYPE_ATTRIB_VOLATILE    = 1 << 2,
+  PDB_TYPE_ATTRIB_UNALIGNED   = 1 << 3,
+  PDB_TYPE_ATTRIB_RESTRICTED  = 1 << 4,
+  PDB_TYPE_ATTRIB_LREF        = 1 << 5,
+  PDB_TYPE_ATTRIB_RREF        = 1 << 6,
+  PDB_TYPE_ATTRIB_FWDREF      = 1 << 7
+} pdb_type_attrib_e;
+typedef U32 pdb_type_attrib;
+
+typedef struct pdb_type_info_proc 
+{
+  /* Calling convention (pdb_cv_call_e) */
+  U32 conv;
+
+  /* itype of the return value. */
+  pdb_cv_itype ret_itype;
+
+  /* Arguments itype. */
+  pdb_cv_itype arg_itype;
+
+  U16 arg_count;
+} pdb_type_info_proc_t;
+
+typedef struct pdb_type_info_method 
+{
+  /* pdb_cv_call_e */
+  U32 conv;           
+
+  /* Return itype. */
+  pdb_cv_itype ret_itype;
+
+  /* Class itype to which this procedure belongs. */
+  pdb_cv_itype class_itype;    
+
+  /* Type of the "this" in C++ */
+  pdb_cv_itype this_itype;
+
+  /* Arguments itype. */
+  pdb_cv_itype arg_itype;      
+
+  U16 arg_count;
+} pdb_type_info_method;
+
+typedef struct pdb_type_info_methodlist
+{
+  pdb_pointer block; /* array of pdb_ml_method */
+} pdb_type_info_methodlist;
+
+typedef struct pdb_type_info_vtshape
+{
+  pdb_uint count;
+  /* pdb_lf_vts_desc are represented as 4bit packed integers */
+  pdb_pointer ptr; 
+} pdb_type_info_vtshape;
+
+typedef struct pdb_type_info_vftable
+{
+  pdb_cv_itype owner_itype;
+  pdb_cv_itype base_table_itype;
+  pdb_uint offset_in_object_layout;
+  pdb_pointer name;
+  pdb_pointer method_names;
+} pdb_type_info_vftable;
+
+typedef struct pdb_type_info_label
+{
+  pdb_cv_ptrmode mode;
+} pdb_type_info_label;
+
+typedef struct pdb_type_info_arglist 
+{
+  // Entry count
+  U32 count;
+
+  // Array of arguments, each argument is described with pdb_cv_itype
+  pdb_pointer itypes; 
+} pdb_type_info_arglist;
+
+typedef struct pdb_type_info_bitfield 
+{
+  /* Underlying type of a bitfield. */
+  pdb_cv_itype base_itype;
+
+  U8 len;
+  U8 pos;
+} pdb_type_info_bitfield;
+
+typedef struct pdb_type_info_udt 
+{
+  /* If this field is a UDT(User Defined Type) this field will have
+   * a non-zero value. */
+  pdb_cv_itype field_itype;
+
+  pdb_cv_itype dims_itype;
+
+  /* Used only to query ENUM_UDT fields. */
+  pdb_cv_itype base_itype;             
+
+  /* If this type belongs to a UDT this field will contains number
+   * of fields. */
+  U16 field_count;
+} pdb_type_info_udt;
+
+typedef struct pdb_type_info_funcid 
+{
+  pdb_pointer name;
+  pdb_cv_itype itype;
+} pdb_type_info_funcid;
+
+typedef struct pdb_type_info_mfuncid
+{
+  pdb_pointer name;          /* method name */
+  pdb_cv_itype parent_itype; /* type index of class (parent) */
+  pdb_cv_itype itype;        /* type index of method */
+} pdb_type_info_mfuncid;
+
+typedef struct pdb_type_info_stringid
+{
+  pdb_pointer   data;       /* string data */
+  pdb_cv_itemid sub_string; /* index of sub string */
+} pdb_type_info_stringid;
+
+typedef struct pdb_type_info_fieldlist
+{
+  pdb_pointer data;
+} pdb_type_info_fieldlist;
+
+typedef struct pdb_type_info_ptr
+{
+  pdb_cv_ptrmode mode;
+  pdb_cv_ptrtype type;
+  pdb_cv_ptr_attrib_t attr;
+} pdb_type_info_ptr;
+
+/*  You can call pdb_infer_itype and pass next_cv_itype to
+ *  resolve type information for next link in a chain of types.
+ *
+ *  If type was declared as 'void', pdb_infer_itype assigns zero to size.
+ *
+ *  When allocating storage for type's value - be careful user can declare any size. 
+ *  This is legal to compile "volatile char buffer[1*1024*1024*1024*1024]".
+ *  Make sure to put a cap on storage allocations to avoid out-of-memory error.
+ */
+typedef struct pdb_type 
+{
+  /* Size of the type in bytes.
+   *
+   * - When type is a pointer size will reflect an arch specific size,
+   *   so 4 bytes on x86 and 8 bytes on x64 and etc.
+   *   
+   * - When type is an array you progress through a chained-link
+   *   of dimensions; more info you can find
+   *   in pdb_tm_infer_ti, search for PDB_LF_ARRAY.
+   */
+  U32 size;
+
+  pdb_type_kind kind;
+
+  pdb_type_attrib attribs;
+
+  /* Unique id that is assigned to the type. */
+  pdb_cv_itype cv_itype;
+
+  /* Type is part of a hierarchy, this field indicates next itype. */
+  pdb_cv_itype next_cv_itype;
+
+  pdb_pointer name;
+
+  union {
+    /* Function ID */
+    pdb_type_info_funcid funcid;
+
+    /* Method ID */
+    pdb_type_info_mfuncid mfuncid;
+
+    /* String ID */
+    pdb_type_info_stringid stringid;
+
+    /* User-defined-type */
+    pdb_type_info_udt udt;
+
+    /* Procedure */
+    pdb_type_info_proc_t proc;
+
+    /* Pointer / Reference */
+    pdb_type_info_ptr ptr;
+
+    /* Method */
+    pdb_type_info_method method;
+
+    /* Virtual table shape */
+    pdb_type_info_vtshape vtshape;
+
+    /* Argument list */
+    pdb_type_info_arglist arglist;
+
+    /* Method list */
+    pdb_type_info_methodlist methodlist;
+
+    /* Bit-field */
+    pdb_type_info_bitfield bitfield;
+    
+    pdb_type_info_fieldlist fieldlist;
+
+    pdb_type_info_vftable vftable;
+
+    pdb_type_info_label label;
+  } u;
+} pdb_type;
+
+typedef struct pdb_tm_bucket 
+{
+  pdb_ti ti;
+  struct pdb_tm_bucket *next;
+} pdb_tm_bucket;
+
+typedef struct pdb_tm 
+{
+  struct pdb_context *pdb;
+  pdb_sn sn;
+  U32 *ti_offsets;
+  pdb_tm_bucket **buckets;
+  pdb_tm_header header;
+} pdb_tm;
+
+PDB_API U32
+pdb_calc_size_for_types(pdb_tm_header *tm_header);
+
+PDB_API syms_bool
+pdb_tm_init(pdb_tm *tm, struct pdb_context *pdb, pdb_default_stream sn, SymsArena *arena);
+
+PDB_API syms_bool
+pdb_tm_offset_for_ti(pdb_tm *tm, pdb_ti ti, U32 *ti_off_out);
+
+PDB_API syms_bool
+pdb_tm_get_itype_off(struct pdb_context *pdb, pdb_cv_itype itype, U32 *itype_off_out);
+
+PDB_API syms_bool
+pdb_tm_get_itemid_off(struct pdb_context *pdb, pdb_cv_itemid itemid, U32 *itemid_off);
+
+PDB_API syms_bool
+pdb_tm_find_ti(pdb_tm *tm, pdb_pointer *name, pdb_ti *ti_out);
+
+PDB_API syms_bool
+pdb_tm_find_itype(struct pdb_context *pdb, pdb_pointer name_ref, pdb_cv_itype *itype_out);
+
+PDB_API syms_bool
+pdb_tm_find_itemid(struct pdb_context *pdb, pdb_pointer name_ref, pdb_cv_itemid *itemid_out);
+
+PDB_API syms_bool
+pdb_tm_infer_ti(pdb_tm *tm, pdb_ti ti, pdb_type *type_out);
+
+PDB_API syms_bool 
+pdb_infer_itype(struct pdb_context *pdb, pdb_cv_itype itype, pdb_type *type_out);
+
+PDB_API syms_bool
+pdb_infer_itemid(struct pdb_context *pdb, pdb_cv_itemid itemid, pdb_type *type_out);
+
+PDB_API syms_bool
+pdb_type_from_name(struct pdb_context *pdb, const char *name, U32 name_len, pdb_type *type);
+
+typedef struct pdb_type_it 
+{
+  struct pdb_context *pdb;
+  pdb_cv_itype next_itype;
+  U32 type_map_index;
+  pdb_tm *type_map[2];
+} pdb_type_it;
+
+typedef struct 
+{
+  pdb_pointer file;
+  pdb_uint ln;
+  pdb_imod mod;
+} pdb_udt_srcline;
+
+PDB_API syms_bool
+pdb_find_udt_srcline(struct pdb_context *pdb, pdb_cv_itype itype, pdb_udt_srcline *srcline_out);
+
+//
+
+typedef struct pdb_cvdata_token 
+{
+  pdb_sn   sn;
+  pdb_uint soffset;
+} pdb_cvdata_token;
+
+typedef struct pdb_context 
+{
+  void *file_data;
+  pdb_uint file_size;
+
+  pdb_uint page_size;
+  pdb_uint free_page_map;
+  pdb_uint pages_used;
+  pdb_uint root_size;
+  U8 page_index_size;
+  union {
+    U16 *addr16;
+    U32 *addr32;
+  } page_map_addr;
+  pdb_uint ver;
+  pdb_uint time;
+  pdb_uint age;
+  /* NOTE(nick): Compiler writes into the PE header same GUID as here. 
+   * You can compare both to determine authenticity. */
+  SymsGUID auth_guid;
+  
+  pdb_sn typed_streams[PDB_TYPED_STREAM_MAX];
+
+  struct {
+    pdb_dbi_header header;
+
+    U32 ver;
+    U16 machine_type;
+
+    U32 modinfo_off;
+    U32 modinfo_len;
+
+    U32 seccon_off;
+    U32 seccon_len;
+
+    U32 secmap_off;
+    U32 secmap_len;
+
+    U32 fileinfo_off;
+    U32 fileinfo_len;
+
+    U32 ecinfo_off;
+    U32 ecinfo_len;
+
+    pdb_sn symrec_sn;
+    pdb_sn pubsym_sn;
+    pdb_sn globalsym_sn;
+
+    pdb_sn dbg_streams[PDB_DBG_STREAM_MAX];
+
+    U16 secs_num;
+    pdb_img_sec *secs;
+
+    /* NOTE(nick): Offsets into the PDB_DEFAULT_STREAM_DBI that point to modules. */
+    U32 mods_num;
+    pdb_uint *mods;
+  } dbi;
+
+  pdb_stream strtable;
+  pdb_stream stroffs;
+
+  struct pdb_tm tpi;
+  struct pdb_tm ipi;
+
+  U32 basic_typenames_array_num;
+  SymsString *basic_typenames_array;
+
+  U32 publics_array_num;
+  pdb_gsi_hr **publics_array;
+
+  U32 globals_array_num;
+  pdb_gsi_hr **globals_array;
+
+  // TODO(nick): trampolines are stored in a light-weight struct "pdb_cvdata_token". 
+  U32 trampoline_count;
+  syms_bool trampoline_contigous;
+  pdb_cvdata_token trampoline_data;
+} pdb_context;
+
+
+/* NOTE(nick; Oct 10; 2019): Offset into the string table. */
+typedef U32 pdb_stroff;
+
+PDB_API U32
+pdb_hashV1_bytes(const void *start, U32 cb, U32 mod);
+
+PDB_API U32
+pdb_hashV1_stream(pdb_stream *stream, U32 size, U32 mod);
+
+PDB_API U32
+pdb_hashV1_pointer(struct pdb_context *pdb, pdb_pointer *bytes, U32 mod);
+
+
+//
+
+typedef enum 
+{
+  PDB_ENCODED_LOCATION_NULL,
+  PDB_ENCODED_LOCATION_ENREGED,
+  PDB_ENCODED_LOCATION_REGREL,
+  PDB_ENCODED_LOCATION_RVA,
+  PDB_ENCODED_LOCATION_VA,
+  PDB_ENCODED_LOCATION_IMPLICIT
+} pdb_encoded_location_type_e;
+typedef U32 pdb_encoded_location_type;
+
+#define PDB_ENCODED_LOCATION_FLAG_NULL          0u
+#define PDB_ENCODED_LOCATION_FLAG_OUTSIDE_RANGE (1u << 0u)
+typedef U32 pdb_encoded_location_flags;
+
+typedef struct pdb_regrel 
+{
+  /* Offset that has to be applied to the register in order to get value for the local. */
+  U32 reg_off;
+
+  /* Index of the register (see pdb_cv_reg_e). */
+  U32 reg_index;
+} pdb_regrel;
+
+/* Compiler stored local in a register. */
+typedef struct pdb_enreged 
+{
+  U32 reg_index;
+} pdb_enreged;
+
+typedef struct pdb_persist 
+{
+  U64 va;
+} pdb_persist;
+
+#define PDB_LOCATION_IMPLICIT_VALUE_MAX 32
+typedef struct pdb_implicit 
+{
+  U8 len;
+  U8 data[PDB_LOCATION_IMPLICIT_VALUE_MAX];
+} pdb_implicit;
+
+/* Represents an encoded location of an address or a value.
+ * Use "pdb_decode_location" to resolve final address. */
+typedef struct pdb_encoded_location 
+{
+  /* Describes type of the memory that union below points to. */
+  pdb_encoded_location_type type;
+  pdb_encoded_location_flags flags;
+
+  union {
+    pdb_enreged  enreged;
+    pdb_regrel   regrel;
+    pdb_persist  persist;
+    pdb_implicit implicit;
+  } u;
+} pdb_encoded_location;
+
+typedef enum 
+{
+  PDB_LOCATION_NULL,
+  PDB_LOCATION_VA,
+  PDB_LOCATION_IMPLICIT
+} pdb_location_e;
+
+typedef struct pdb_location 
+{
+  pdb_location_e type;
+  union {
+    SymsAddr va;
+    pdb_implicit implicit;
+  } u;
+} pdb_location;
+
+PDB_API pdb_encoded_location
+pdb_encode_location_for_secoff(struct pdb_context *pdb, U32 sec, U32 off);
+
+PDB_API pdb_encoded_location
+pdb_encode_location_for_regrel(struct pdb_context *pdb, U32 regindex, U32 regoff);
+
+PDB_API pdb_encoded_location
+pdb_encode_location_for_datasym32(struct pdb_context *pdb, pdb_cv_datasym32 *datasym);
+
+#define PDB_MEMREAD_SIG(name) syms_bool name(void *context, SymsAddr va, void *read_buffer, U32 read_buffer_max)
+typedef PDB_MEMREAD_SIG(pdb_memread_sig);
+
+#define PDB_REGREAD_SIG(name) U32 name(void *context, U32 reg_index, void *read_buffer, U32 read_buffer_max)
+typedef PDB_REGREAD_SIG(pdb_regread_sig);
+
+#define PDB_REGWRITE_SIG(name) U32 name(void *context, U32 reg_index, void *write_buffer, U32 write_buffer_max)
+typedef PDB_REGWRITE_SIG(pdb_regwrite_sig);
+
+PDB_API syms_bool
+pdb_decode_location(pdb_encoded_location *encoded_loc,
+    SymsAddr orig_rebase,
+    SymsAddr rebase,
+    void *memread_ctx, pdb_memread_sig *memread,
+    void *regread_ctx, pdb_regread_sig *regread,
+    pdb_location *decoded_loc);
+
+PDB_API syms_bool
+pdb_build_va(struct pdb_context *pdb, U32 sec, U32 off, SymsAddr *out_va);
+
+PDB_API syms_bool
+pdb_build_sec_off(struct pdb_context *pdb, SymsAddr va, pdb_isec *sec, pdb_isec_umm *off);
+
+/* -------------------------------------------------------------------------------- */
+
+typedef struct pdb_var 
+{
+  /* Type of the local, you can pass this to the "pdb_infer_itype". */
+  pdb_cv_itype itype;
+
+  pdb_cv_localsym_flags flags;
+
+  /* Encoded address or value where variable is located in memory. */
+  pdb_encoded_location encoded_va;
+
+  pdb_pointer name;
+
+  pdb_pointer gaps;
+} pdb_var;
+
+typedef struct pdb_sec_it 
+{
+  struct pdb_context *pdb;
+  pdb_stream stream;
+} pdb_sec_it;
+
+PDB_API syms_bool
+pdb_sec_it_init(struct pdb_context *pdb, pdb_sec_it *sec_it);
+
+PDB_API syms_bool
+pdb_sec_it_next(pdb_sec_it *sec_it, pdb_img_sec *sec);
+
+PDB_API syms_bool
+pdb_sec_from_index(struct pdb_context *pdb, U32 index, pdb_img_sec *sec_out);
+
+typedef struct pdb_debug_sec 
+{
+  pdb_cv_ss_type type;
+  pdb_stream stream;
+} pdb_debug_sec;
+
+typedef enum 
+{
+  PDB_MOD_SEC_SYMS,
+  PDB_MOD_SEC_LINES_C11,
+  PDB_MOD_SEC_LINES_C13,
+  PDB_MOD_SEC_INLINE_LINES
+} pdb_mod_sec_type_e;
+
+typedef enum 
+{
+  PDB_MOD_FILECHKSUM_CACHED = (1 << 0)
+} pdb_mod_flags_e;
+typedef U32 pdb_mod_flags;
+
+typedef struct pdb_mod 
+{
+  struct pdb_context *pdb;
+
+  pdb_imod id;
+
+  pdb_mod_flags flags;
+
+  pdb_sn sn;
+
+  pdb_isec sec;
+  pdb_isec_umm sec_off;
+  pdb_isec_umm sec_len;
+
+  U32 syms_size;
+  U32 c11_lines_size;
+  U32 c13_lines_size;
+
+  pdb_pointer name;
+  pdb_pointer name2;
+
+  pdb_debug_sec filechksum;
+} pdb_mod;
+
+typedef struct pdb_mod_it 
+{
+  struct pdb_context *pdb;
+  pdb_stream dbi_data;
+  pdb_imod imod;
+} pdb_mod_it;
+
+PDB_API syms_bool
+pdb_mod_it_init(pdb_mod_it *mod_it, struct pdb_context *pdb);
+
+PDB_API syms_bool
+pdb_mod_it_next(pdb_mod_it *mod_it, pdb_mod *mod_out);
+
+PDB_API syms_bool
+pdb_mod_it_seek(pdb_mod_it *mod_it, pdb_imod imod);
+
+PDB_API syms_bool
+pdb_mod_get_debug_sec(pdb_mod *mod, pdb_mod_sec_type_e sec, pdb_stream *dbi_data_out);
+
+PDB_API syms_bool
+pdb_imod_from_isec(struct pdb_context *pdb, pdb_isec sec, pdb_isec_umm off, pdb_imod *imod_out);
+
+PDB_API syms_bool
+pdb_mod_init(pdb_mod *mod_out, struct pdb_context *pdb, pdb_imod imod);
+
+typedef struct pdb_file_info 
+{
+  pdb_pointer         path;
+  pdb_pointer         chksum;
+  pdb_cv_chksum_type  chksum_type;
+} pdb_file_info; // 72->56 bytes
+
+PDB_API syms_bool
+pdb_mod_get_filechksum(pdb_mod *mod, pdb_stream *stream_out);
+
+typedef struct pdb_debug_sec_it 
+{
+  pdb_stream stream;
+} pdb_debug_sec_it;
+
+PDB_API syms_bool
+pdb_debug_sec_it_init(pdb_debug_sec_it *it, pdb_mod *mod);
+
+PDB_API syms_bool
+pdb_debug_sec_it_next(pdb_debug_sec_it *it, pdb_debug_sec *sec_out);
+
+typedef struct pdb_dss_it 
+{
+  pdb_cv_ss_type type;
+  pdb_stream stream;
+  syms_bool ex_mode;
+} pdb_dss_it;
+
+/* Inline data is stored in DBI stream and this is header file for it. */
+typedef struct pdb_ss_inline 
+{
+  pdb_cv_itemid inlinee;
+
+  U32 src_ln;
+  U32 file_id;
+  U32 extra_files_count;
+  pdb_pointer extra_files;
+} pdb_ss_inline;
+
+PDB_API syms_bool
+pdb_dss_it_init(pdb_dss_it *it_out, pdb_debug_sec *sec);
+
+PDB_API syms_bool
+pdb_dss_it_next_inline(pdb_dss_it *it, pdb_ss_inline *inline_out);
+
+typedef struct pdb_sym_it 
+{
+  syms_bool inited_from_token;
+  pdb_stream stream;
+} pdb_sym_it;
+
+PDB_API syms_bool
+pdb_sym_it_init(pdb_sym_it *sym_it, pdb_mod *mod);
+
+PDB_API syms_bool
+pdb_sym_it_init_token(pdb_sym_it *sym_it, struct pdb_context *pdb, pdb_cvdata_token token);
+
+PDB_API syms_bool
+pdb_sym_it_read(pdb_sym_it *sym_it, pdb_cv_sym_type *type_out, pdb_stream *stream_out);
+
+PDB_API syms_bool
+pdb_sym_it_peek(pdb_sym_it *sym_it, pdb_cv_sym_type *type_out, pdb_stream *stream_out);
+
+PDB_API syms_bool
+pdb_sym_it_next(pdb_sym_it *sym_it, pdb_cv_sym_type *type_out, pdb_stream *stream_out);
+
+enum
+{
+  PDB_LOCAL_EXPORT_NULL,
+  PDB_LOCAL_EXPORT_VAR,
+  PDB_LOCAL_EXPORT_SCOPE,
+  PDB_LOCAL_EXPORT_SCOPE_END
+};
+typedef pdb_uint pdb_local_export_type;
+
+typedef struct pdb_scope
+{
+  SymsAddr inst_lo;
+  SymsAddr inst_hi;
+} pdb_scope;
+
+typedef struct pdb_local_export
+{
+  pdb_local_export_type type;
+  union
+  {
+    pdb_var var;
+    pdb_scope scope;
+  } u;
+} pdb_local_export;
+
+typedef struct pdb_local_it 
+{
+  struct pdb_context *pdb;
+  syms_bool defrange_mode;
+  pdb_sym_it sym_it;
+
+  pdb_isec sec;
+  pdb_isec_umm sec_off;
+
+  U32 range_off;
+  U32 range_len;
+
+  S32 block32_count;
+  S32 regrel32_count;
+  S32 inlinesite_count;
+
+  U16 arg_count;
+} pdb_local_it;
+
+PDB_API syms_bool
+pdb_local_it_init_(pdb_local_it *local_it, struct pdb_context *pdb, pdb_isec sec, pdb_isec_umm sec_off, pdb_cvdata_token cvdata);
+
+PDB_API syms_bool
+pdb_local_it_init(pdb_local_it *local_it, struct pdb_context *pdb, pdb_isec sec, pdb_isec_umm sec_off);
+
+PDB_API syms_bool
+pdb_local_it_next(pdb_local_it *local_it, pdb_local_export *export_out);
+
+typedef struct pdb_global_it 
+{
+  struct pdb_context *pdb;
+  pdb_stream stream;
+  U32 hr_index;
+  pdb_gsi_hr *hr;
+} pdb_global_it;
+
+PDB_API syms_bool
+pdb_global_it_init(pdb_global_it *global_it, struct pdb_context *pdb);
+
+PDB_API syms_bool
+pdb_global_it_next(pdb_global_it *global_it, pdb_var *var_out);
+
+PDB_API syms_bool
+pdb_global_from_name(struct pdb_context *pdb, const char *name, U32 name_size, pdb_var *global_out);
+
+typedef struct pdb_file_it 
+{
+  struct pdb_context *pdb;
+  pdb_stream strtable;
+  pdb_stream stroffs;
+  U32 off_count;
+  U32 num_read;
+} pdb_file_it;
+
+PDB_API syms_bool
+pdb_file_it_init(pdb_file_it *file_it, struct pdb_context *pdb);
+
+PDB_API syms_bool
+pdb_file_it_next(pdb_file_it *file_it, pdb_pointer *file_out);
+
+typedef struct pdb_proc_it 
+{
+  pdb_sym_it sym_it;
+} pdb_proc_it;
+
+typedef struct pdb_proc 
+{
+  pdb_isec sec;
+  pdb_isec_umm sec_off;
+
+  /* Length in bytes, of all instruction that contribute to a procedure. */
+  pdb_uint size;
+
+  pdb_cv_itype itype;
+
+  /* pdb_cv_proc_flags_e */
+  pdb_uint flags;
+
+  pdb_string_ref name;
+
+  pdb_cvdata_token cvdata;
+} pdb_proc;
+
+PDB_API syms_bool
+pdb_proc_it_init(pdb_proc_it *proc_it, pdb_mod *mod);
+
+PDB_API syms_bool
+pdb_proc_it_next(pdb_proc_it *proc_it, pdb_proc *out_proc);
+
+PDB_API syms_bool
+pdb_proc_from_name(pdb_context *pdb, const char *name, pdb_uint name_len, pdb_proc *proc_out);
+
+PDB_API syms_bool
+pdb_proc_from_cvdata(pdb_context *pdb, pdb_cvdata_token cvdata, pdb_proc *proc_out);
+
+PDB_API syms_bool
+pdb_proc_from_stream(pdb_cv_sym_type cvtype, pdb_stream *cvdata, pdb_proc *proc_out);
+
+PDB_API syms_bool
+pdb_proc_from_name_(pdb_context *pdb, pdb_pointer *name, pdb_proc *proc_out);
+
+typedef struct pdb_inline_site 
+{
+  pdb_pointer name;
+
+  pdb_cv_itype itype;
+
+  pdb_isec sec;
+  pdb_isec_umm sec_off;
+  u32 size;
+
+  pdb_file_info fi;
+
+  U32 ln_at_pc;
+
+  pdb_cvdata_token cvdata;
+} pdb_inline_site;
+
+typedef struct pdb_inline_it 
+{
+  S32 site_count;
+  /* NOTE(nick; Oct 6 2019): Procedure relative Program Counter. */
+  SymsAddr proc_pc;
+  pdb_proc proc;
+  pdb_mod mod;
+  pdb_sym_it sym_it;
+} pdb_inline_it;
+
+PDB_API syms_bool
+pdb_inline_it_init(pdb_inline_it *inline_it, struct pdb_context *pdb, pdb_cvdata_token cvproc, SymsAddr proc_rva);
+
+PDB_API syms_bool
+pdb_inline_it_read(pdb_inline_it *inline_it, pdb_cv_inlinesym *sym_out, pdb_stream *ba_out);
+
+PDB_API syms_bool
+pdb_inline_it_next(pdb_inline_it *inline_it, pdb_inline_site *site_out);
+
+typedef struct pdb_const_value 
+{
+  pdb_string_ref name;
+  pdb_cv_itype itype;
+  pdb_numeric value;
+} pdb_const_value;
+
+typedef struct pdb_const_it 
+{
+  struct pdb_context *pdb;
+  pdb_stream stream;
+  U32 index;
+  pdb_gsi_hr *hr;
+} pdb_const_it;
+
+PDB_API syms_bool
+pdb_const_it_init(pdb_const_it *const_it, struct pdb_context *pdb);
+
+PDB_API syms_bool
+pdb_const_it_next(pdb_const_it *const_it, pdb_const_value *const_out);
+
+PDB_API syms_bool
+pdb_const_from_name(struct pdb_context *pdb, const char *name, U32 name_size, pdb_const_value *const_out);
+
+typedef struct pdb_arg_it 
+{
+  /* parser context */
+  struct pdb_context *pdb;
+
+  /* current argument index */
+  u32 idx;
+
+  /* array of pdb_cv_itype entries */
+  pdb_pointer itypes;
+} pdb_arg_it;
+
+PDB_API syms_bool
+pdb_arg_it_init(pdb_arg_it *it, struct pdb_context *pdb, pdb_cv_itype arglist_itype);
+
+PDB_API syms_bool
+pdb_arg_it_next(pdb_arg_it *it, pdb_cv_itype *out_arg_type);
+
+/* -------------------------------------------------------------------------------- */
+
+enum 
+{
+  PDB_MEMBER_TYPE_NULL,
+
+  /* Member for data that is stored inside the UDT  */
+  PDB_MEMBER_TYPE_DATA,
+
+  /* Class method */
+  PDB_MEMBER_TYPE_METHOD,
+
+  /* Enum member */
+  PDB_MEMBER_TYPE_ENUMERATOR,
+
+  /* Type that was declared inside a class */
+  PDB_MEMBER_TYPE_NESTED_TYPE,
+
+  /* Inherited class */
+  PDB_MEMBER_TYPE_BASE_CLASS,
+
+  /* Describes memory region for virtual functions.
+   * Memory is packed as an array of 4 bit integers,
+   * each integer maps into pdb_lf_vts_desc. */
+  PDB_MEMBER_TYPE_VIRTUAL_TABLE,
+  
+  /* Static data member of a class */
+  PDB_MEMBER_TYPE_STATIC_DATA,
+
+  /* 
+   * CodeView packs methods with matching identifiers into 
+   * an array of methods. I am sorry for ugly hack - we need 
+   * it to hide in the interface layer codeview shenanigans. 
+   *
+   * Also we dont have an allocator in the API layer, 
+   * iterators, symbols, line info, and so on are allocated 
+   * on the stack. Library allocates memory just for lookup
+   * acceleration and it is possible to parse entire pdb
+   * with zero memory allocations! 
+   *
+   * */
+  PDB_MEMBER_TYPE_METHODLIST
+};
+typedef u32 pdb_member_type;
+
+/* Contains anything that resembles data member in struct, class, enum, and union */
+typedef struct pdb_member 
+{
+  pdb_pointer name;
+
+  pdb_member_type type;
+
+  pdb_cv_itype itype;
+
+  pdb_cv_fldattr_t attr;
+
+  union 
+  {
+    struct
+    {
+      /* Parent-relative offset where this member is located in memory. */
+      U32 offset; 
+    } data;
+
+    struct 
+    {
+      u32 count;
+    } methodlist;
+
+    struct
+    {
+      pdb_uint vbaseoff;
+    } method;
+
+    struct 
+    {
+      /* Use this field only if the "type" field is set to the PDB_MEMBER_TYPE_ENUMERATOR.
+       * This field indicates value that was assigned to enum field. */
+      pdb_numeric value;
+    } enumerator;
+
+    struct 
+    {
+      u32 offset;
+    } base_class;
+
+    struct
+    {
+      u32 offset;
+    } vtab;
+  } u;
+} pdb_member;
+
+typedef struct pdb_member_it 
+{
+  U32 stream_end;
+  struct pdb_context *pdb;
+  pdb_type_info_udt *udt;
+  pdb_stream stream;
+} pdb_member_it;
+
+PDB_API syms_bool
+pdb_member_it_init(pdb_member_it *member_it, struct pdb_context *pdb, pdb_type_info_udt *udt);
+
+PDB_API syms_bool
+pdb_member_it_next(pdb_member_it *member_it, pdb_member *member_out);
+
+/* -------------------------------------------------------------------------------- */
+
+/* Indicates if we need to read another line section. */
+#define PDB_LINE_IT_FLAGS_NEW_SECTION (1u << 0u)
+typedef U32 pdb_line_it_flags;
+
+typedef enum 
+{
+  PDB_LINE_FORMAT_NULL,
+  PDB_LINE_FORMAT_C11,
+  PDB_LINE_FORMAT_C13
+} pdb_line_format;
+
+typedef struct pdb_line_it 
+{
+  /* Iterator for the modules. */
+  pdb_mod mod;
+  pdb_file_info fi;
+  pdb_line_format format;
+  pdb_line_it_flags flags;
+  pdb_stream stream;
+
+  U16 sec;
+  U32 sec_off;
+  U32 sec_size;
+
+  U32 last_read_ln;
+
+  union {
+    struct {
+      /* Count for the total number of lines in current section. */
+      U32 line_index_max;
+
+      /* When a new module comes, we scan it to get offset of the checksum 
+       * section because all lines instead of providing offset where the 
+       * file name is in the name-stream, they give offset into the 
+       * checksum section and from there we have to resolve the file name
+       * for given lines block. */
+      U32 chksum_off;
+
+      /* Offset where current line section ends. */
+      U32 sec_end;
+
+      pdb_stream cv_lines;
+    } c13;
+
+    struct {
+      U32 file_index;
+      U32 file_count;
+
+      U16 filesec_index;
+      U16 filesec_count;
+
+      U16 pair_index;
+      U16 pair_count;
+
+      pdb_stream secrange_stream;
+      pdb_stream off_stream;
+      pdb_stream ln_stream;
+    } c11;
+  } u;
+} pdb_line_it;
+
+typedef struct pdb_line 
+{
+  U16 sec;
+  U32 off;
+  U32 ln;
+} pdb_line;
+
+PDB_API syms_bool
+pdb_line_it_init(pdb_line_it *line_it, pdb_mod *mod);
+
+PDB_API syms_bool
+pdb_line_it_next(pdb_line_it *line_it, pdb_line *line);
+
+typedef struct pdb_map 
+{
+  /* section and offset for the first instruction in the source line. */
+  pdb_isec sec;
+  pdb_isec_umm sec_off;
+
+  /* Size of all instructions that contribute to the source line. */
+  U32 size;
+
+  /* Line number to which VA was mapped. */
+  U32 ln;
+
+  pdb_file_info fi;
+} pdb_map;
+
+PDB_API syms_bool
+pdb_src_to_va(struct pdb_context *pdb, char *file_name, U32 file_name_len, U32 ln,
+        pdb_map *map);
+
+PDB_API syms_bool
+pdb_va_to_src(struct pdb_context *pdb, SymsAddr line_va, pdb_map *map);
+
+/* -------------------------------------------------------------------------------- */
+
+PDB_API syms_bool
+pdb_find_nearest_sym(struct pdb_context *pdb, SymsAddr va, pdb_pointer *name_out);
+
+PDB_API syms_bool
+pdb_find_nearest_sc(struct pdb_context *pdb, SymsAddr va, pdb_sc *sc_out);
+
+PDB_API syms_bool
+pdb_find_nearest_sc_ex(struct pdb_context *pdb, U32 sec, U32 off, pdb_sc *sc_out);
+
+// File Info section of DBI packs unique strings of source paths. It has a 16bit counter
+// -- hello from good old days -- and today medium size program that link with CRT
+// have 30K files when you unpack file info section.
+typedef struct pdb_fileinfo
+{
+  u32 mod_count;
+  u32 src_count;
+  pdb_stream imod_block;        // pdb_uint16[mod_count] module indices
+  pdb_stream count_block;       // source count per module, counts are packed as variable-length arrays.
+  pdb_stream ich_block;         // offsets into str_block
+  pdb_stream str_block;         // char * array
+} pdb_fileinfo;
+
+PDB_API syms_bool
+pdb_fileinfo_init(struct pdb_context *pdb, pdb_fileinfo *fi);
+
+PDB_API syms_bool
+pdb_fileinfo_get_src_count(pdb_fileinfo *fi, pdb_imod imod, pdb_uint *count_out);
+
+PDB_API syms_bool
+pdb_fileinfo_lookup(pdb_fileinfo *fi, pdb_imod imod, pdb_uint index, pdb_pointer *cptr_out);
+
+typedef struct pdb_strtable
+{
+  u32 magic;
+  u32 version;
+  u32 bucket_count;
+  pdb_stream strblock;
+  pdb_stream buckets;
+} pdb_strtable;
+
+PDB_API syms_bool
+pdb_init_strtable(struct pdb_strtable *st, struct pdb_stream *stream);
+
+PDB_API pdb_stream
+pdb_get_strtable(struct pdb_context *pdb);
+
+PDB_API pdb_stream
+pdb_get_stroffs(struct pdb_context *pdb);
+
+PDB_API syms_bool
+pdb_strtable_at(struct pdb_context *pdb, pdb_ni ni, pdb_pointer *str_out);
+
+PDB_API syms_bool
+pdb_strtable_off_to_str(struct pdb_context *pdb, pdb_stroff stroff, pdb_pointer *str_out);
+
+PDB_API syms_bool
+pdb_strtable_find(struct pdb_context *pdb, pdb_pointer *name, pdb_strcmp_flags_e cmp_flags, pdb_ni *ni);
+
+PDB_API syms_bool
+pdb_strtable_find_str(struct pdb_context *pdb, SymsString name, pdb_strcmp_flags_e cmp_flags, pdb_ni *ni);
+
+
+PDB_API U32
+pdb_parse_gsi_hash_table(pdb_stream *gsi_stream, SymsArena *arena, pdb_gsi_hr ***out_table, U32 *out_num_table);
+
+PDB_API syms_bool
+pdb_init(struct pdb_context *pdb, void *data, U32 data_size);
+
+PDB_API SymsUMM
+pdb_load_types(pdb_context *pdb, struct SymsArena *arena);
+
+PDB_API SymsUMM
+pdb_load_publics(pdb_context *pdb, struct SymsArena *arena);
+
+PDB_API SymsUMM
+pdb_load_globals(pdb_context *pdb, struct SymsArena *arena);
+
+PDB_API SymsUMM
+pdb_load_dbi(pdb_context *pdb, struct SymsArena *arena);
+
+PDB_API SymsUMM
+pdb_load_trampolines(pdb_context *pdb, struct SymsArena *arena);
+
+PDB_API SymsUMM
+pdb_calc_size_for_aux_data(pdb_context *pdb);
+
+PDB_API syms_bool
+pdb_load_aux_data(pdb_context *pdb, SymsArena *arena);
+
+PDB_API syms_bool
+pdb_trampoline_from_ip(pdb_context *pdb, pdb_isec sec, pdb_isec_umm sec_off, pdb_isec *sec_out, pdb_isec_umm *off_out);
+
+PDB_API syms_bool
+pdb_stream_read_numeric(pdb_stream *stream, pdb_numeric *out_num);
+
+PDB_API syms_bool
+pdb_stream_read_numeric_u32(pdb_stream *stream, U32 *out_value);
+
+PDB_API syms_bool
+pdb_stream_read_numeric_u64(pdb_stream *stream, U64 *out_value);
+
+PDB_API SymsNTFileHeaderMachineType
+pdb_get_machine_type(struct pdb_context *pdb);
+
+PDB_API syms_bool
+pdb_type_it_init(pdb_type_it *it, struct pdb_context *pdb);
+
+PDB_API syms_bool
+pdb_type_it_next(pdb_type_it *it, pdb_cv_itype *itype_out);
+
+PDB_API const char *
+pdb_ver_to_str(pdb_context *pdb);
+
+PDB_API const char *
+pdb_dbi_ver_to_str(pdb_context *pdb);
+
+PDB_API const char *
+pdb_basic_itype_to_str(pdb_basic_type itype);
+
+PDB_API syms_bool
+pdb_fileinfo_get_strblock(pdb_fileinfo *fi, pdb_imod imod, u32 *count_out, pdb_stream *strblock_out);
+
+#endif /* SYMS_PDB_INCLUDE_H */
diff --git a/Engine/Source/Runtime/SymsLib/include/pdb/syms_pdb_api.h b/Engine/Source/Runtime/SymsLib/include/pdb/syms_pdb_api.h
new file mode 100644
--- /dev/null
+++ b/Engine/Source/Runtime/SymsLib/include/pdb/syms_pdb_api.h
@@ -0,0 +1,154 @@
+// Copyright Epic Games, Inc. All Rights Reserved.
+#ifndef SYMS_SYMS_API_INCLUDE_H
+#define SYMS_SYMS_API_INCLUDE_H
+
+/******************************************************************************
+ * File   : syms_pdb.h                                                        *
+ * Author : Nikita Smith                                                      *
+ * Created: 2020/06/20                                                        *
+ * Purpose: Declarations for pdb.h interface                                  *
+ ******************************************************************************/
+
+typedef struct SymsDebugInfoPdb 
+{
+  pdb_context context;
+} SymsDebugInfoPdb;
+
+typedef struct SymsModIterPdb 
+{
+  pdb_mod_it impl;
+} SymsModIterPdb;
+
+typedef struct SymsLineIterPdb 
+{
+  pdb_line_it impl;
+} SymsLineIterPdb;
+
+typedef struct SymsMemberIterPdb 
+{
+  pdb_member_it impl;
+  pdb_pointer methodlist_name;
+  pdb_uint methodlist_index;
+  pdb_uint methodlist_count;
+  pdb_pointer methodlist_block;
+} SymsMemberIterPdb;
+
+typedef struct SymsGlobalIterPdb 
+{
+  pdb_global_it impl;
+} SymsGlobalIterPdb;
+
+typedef struct SymsLocalDataIterPdb 
+{
+  pdb_sym_it impl;
+} SymsLocalDataIterPdb;
+
+typedef struct SymsFileIterPdb 
+{
+  pdb_file_it impl;
+} SymsFileIterPdb;
+
+typedef struct SymsProcIterPdb 
+{
+  pdb_proc_it impl;
+} SymsProcIterPdb;
+
+typedef struct SymsArgIterPdb 
+{
+  pdb_arg_it impl;
+} SymsArgIterPdb;
+
+typedef struct SymsLocalIterPdb 
+{
+  pdb_local_it impl;
+  SymsScope *scope_stack;
+  syms_uint scope_count;
+  syms_uint scope_max;
+} SymsLocalIterPdb;
+
+typedef struct SymsInlineIterPdb 
+{
+  pdb_inline_it impl;
+} SymsInlineIterPdb;
+
+typedef struct SymsConstIterPdb 
+{
+  pdb_const_it impl;
+} SymsConstIterPdb;
+
+typedef struct SymsTypeIterPdb 
+{
+  pdb_type_it impl;
+} SymsTypeIterPdb;
+
+typedef struct SymsRangePdb
+{
+  SymsAddr lo;
+  SymsAddr hi;
+  pdb_pointer gaps;
+} SymsRangePdb;
+
+typedef struct SymsRangeIterPdb
+{
+  pdb_context *context;
+  pdb_pointer gaps;
+  SymsAddr lo;
+  SymsAddr hi;
+  syms_uint read_offset;
+  syms_bool last_range_emitted;
+} SymsRangeIterPdb;
+
+SYMS_COMPILER_ASSERT(sizeof(SymsRangePdb) <= sizeof(SymsRangeImpl));
+SYMS_COMPILER_ASSERT(sizeof(SymsRangeIterPdb) <= sizeof(SymsRangeIterImpl));
+SYMS_COMPILER_ASSERT(sizeof(pdb_type) <= sizeof(SymsTypeImpl));
+SYMS_COMPILER_ASSERT(sizeof(pdb_mod) <= sizeof(SymsModImpl));
+SYMS_COMPILER_ASSERT(sizeof(SymsDebugInfoPdb) <= sizeof(SymsDebugInfoImpl));
+SYMS_COMPILER_ASSERT(sizeof(SymsModIterPdb) <= sizeof(SymsModIterImpl));
+SYMS_COMPILER_ASSERT(sizeof(SymsLineIterPdb) <= sizeof(SymsLineIterImpl));
+SYMS_COMPILER_ASSERT(sizeof(SymsMemberIterPdb) <= sizeof(SymsMemberIterImpl));
+SYMS_COMPILER_ASSERT(sizeof(SymsGlobalIterPdb) <= sizeof(SymsGlobalIterImpl));
+SYMS_COMPILER_ASSERT(sizeof(SymsLocalDataIterPdb) <= sizeof(SymsLocalDataIterImpl));
+SYMS_COMPILER_ASSERT(sizeof(SymsFileIterPdb) <= sizeof(SymsFileIterImpl));
+SYMS_COMPILER_ASSERT(sizeof(SymsProcIterPdb) <= sizeof(SymsProcIterImpl));
+SYMS_COMPILER_ASSERT(sizeof(SymsArgIterPdb) <= sizeof(SymsArgIterImpl));
+SYMS_COMPILER_ASSERT(sizeof(SymsLocalIterPdb) <= sizeof(SymsLocalIterImpl));
+SYMS_COMPILER_ASSERT(sizeof(SymsInlineIterPdb) <= sizeof(SymsInlineIterImpl));
+SYMS_COMPILER_ASSERT(SYMS_LOCATION_IMPLICIT_VALUE_MAX >= PDB_LOCATION_IMPLICIT_VALUE_MAX);
+SYMS_COMPILER_ASSERT(SYMS_MEMBER_ENUM_MAX >= PDB_NUMERIC_MAX);
+SYMS_COMPILER_ASSERT(SYMS_CONST_VALUE_MAX >= PDB_NUMERIC_MAX);
+SYMS_COMPILER_ASSERT(sizeof(SymsStringRefImpl) >= sizeof(pdb_pointer));
+
+typedef PDB_MEMREAD_SIG(SymsPdbMemreadSig);
+typedef PDB_REGREAD_SIG(SymsPdbRegreadSig);
+typedef PDB_REGWRITE_SIG(SymsPdbRegwriteSig);
+
+// ----------------------------------------------------------------------------
+
+SYMS_INTERNAL pdb_cv_itype
+syms_typeid_to_pdb(const SymsTypeID *type_id);
+
+SYMS_INTERNAL SymsTypeID
+syms_typeid_for_pdb(pdb_cv_itype itype);
+
+SYMS_INTERNAL syms_int
+syms_typeid_cmp_pdb(pdb_cv_itype l, pdb_cv_itype r);
+
+// ----------------------------------------------------------------------------
+
+SYMS_INTERNAL SymsStringRef
+syms_string_ref_pdb(pdb_pointer *pointer);
+
+SYMS_INTERNAL syms_bool
+syms_proc_from_pdb_proc(SymsDebugInfoPdb *debug_info, pdb_proc *proc, SymsProc *proc_out);
+
+SYMS_INTERNAL void
+syms_const_convert_from_pdb(pdb_const_value *pdb_const, SymsConst *const_out);
+
+SYMS_INTERNAL syms_bool
+syms_pdb_type_to_syms_type(pdb_context *pdb, pdb_type *type, SymsType *type_out);
+
+SYMS_INTERNAL
+PDB_REGREAD_SIG(syms_regread_pdb);
+
+#endif /* SYMS_SYMS_API_INCLUDE_H */
+
diff --git a/Engine/Source/Runtime/SymsLib/include/regs/syms_regs.h b/Engine/Source/Runtime/SymsLib/include/regs/syms_regs.h
new file mode 100644
--- /dev/null
+++ b/Engine/Source/Runtime/SymsLib/include/regs/syms_regs.h
@@ -0,0 +1,179 @@
+// Copyright Epic Games, Inc. All Rights Reserved.
+#ifndef SYMS_REGS_INCLUDE_H
+#define SYMS_REGS_INCLUDE_H
+
+/******************************************************************************
+ * File   : syms_regs.h                                                       *
+ * Author : Nikita Smith                                                      *
+ * Created: 2020/06/10                                                        *
+ * Purpose: Generic register interface that supports x86/x64                  *
+ ******************************************************************************/
+
+#ifndef REG_API
+#define REG_API SYMS_API
+#endif
+
+typedef U32 SymsRegID;
+
+typedef struct SymsRegDesc 
+{
+    // register name/alias
+    char name[8];
+    
+    // register class
+    U16 regclass;
+    
+    // register contained within a variable at this byte offset into the register file 
+    U16 varOffset_le;
+    
+    // bit offset within the containing variable 
+    U16 bitpos_le;
+    
+    // bit size within the containing variable 
+    U16 bitcount_le;
+} SymsRegDesc;
+SYMS_COMPILER_ASSERT(sizeof(SymsRegDesc) == 16);
+
+typedef struct SymsRegsPtr 
+{
+    U8 *data;
+    U16 bytes;
+} SymsRegsPtr;
+
+typedef struct SymsRegs 
+{
+    U16 endiness;
+    
+    // SymsArch type
+    U16 arch_le;
+    
+    // Regs version
+    U16 version_le;
+    
+    // Lowest register
+    U32 firstreg_le;
+    
+    // Highest register
+    U32 lastreg_le;
+    
+    /* Total register file size */
+    U32 size_le;
+    
+    SymsRegDesc descs[0];
+} SymsRegs;
+
+typedef struct
+{
+    SymsRegs *regs;
+    
+    U32 desc_count;
+    SymsRegDesc *desc;
+    
+    U32 values_size;
+    U8 *values;
+} SymsRegsInfo;
+
+#define SYMS_REG_ID_ARCH_MASK(x) ((U32)(x) >> 12)
+
+syms_bool
+syms_get_table_info(SymsArch arch, SymsRegDesc **desc_out, U32 *desc_size_out, U32 *regs_size_out, U32 *ver_out, U32 *regid_min_out, U32 *regid_max_out);
+
+REG_API U32
+syms_regs_size_for_arch(SymsArch arch);
+
+REG_API SymsRegs *
+syms_regs_init(SymsArch arch, void *bf, U32 bf_max);
+
+REG_API SymsRegs *
+syms_regs_clone(const SymsRegs *regs, void *buffer, U32 buffer_max);
+
+// Returns the CPU architecture for a register file.
+REG_API SymsArch 
+syms_regs_get_arch(const SymsRegs *regs);
+
+// Returns the size of a register file, in bytes.
+REG_API U32 
+syms_regs_get_size(const SymsRegs *regs);
+
+// Returns the version number of a register file.
+REG_API U32 
+syms_regs_get_version(const SymsRegs *regs);
+
+// Returns the lowest register ID for a file. 
+REG_API SymsRegID 
+syms_regs_get_first_regid(const SymsRegs *regs);
+
+// Returns the highest register ID for a file.
+REG_API SymsRegID 
+syms_regs_get_last_regid(const SymsRegs *regs);
+
+// Returns the register ID for the instruction pointer.
+REG_API SymsRegID 
+syms_regs_get_ip_regid(const SymsRegs *regs);
+
+// Returns the register ID for the stack pointer.
+REG_API SymsRegID
+syms_regs_get_sp_regid(const SymsRegs *regs);
+
+// Returns the description for a register file & index.
+REG_API const SymsRegDesc *
+syms_regs_get_regdesc(const SymsRegs *regs, SymsRegID id);
+
+REG_API SymsRegID
+syms_regid_from_name(const SymsRegs *regs, char *name);
+
+REG_API U32 
+syms_regs_get_value(const SymsRegs *regs, SymsRegID id, void *dst, U32 dst_size);
+
+REG_API U32 
+syms_regs_set_value(SymsRegs *regs, SymsRegID id, const void *src, U32 src_size);
+
+REG_API void 
+syms_regs_clear_reg(SymsRegs *regs, SymsRegID id);
+
+REG_API U8  
+syms_regs_get8(const SymsRegs *regs, SymsRegID id);
+
+REG_API U16 
+syms_regs_get16(const SymsRegs *regs, SymsRegID id);
+
+REG_API U32 
+syms_regs_get32(const SymsRegs *regs, SymsRegID id);
+
+REG_API U64 
+syms_regs_get64(const SymsRegs *regs, SymsRegID id);
+
+REG_API syms_bool 
+syms_regs_get_addr(const SymsRegs *regs, SymsRegID id, SymsAddr *addr_out);
+
+REG_API syms_bool 
+syms_regs_set8(SymsRegs *regs, SymsRegID id, U8 value);
+
+REG_API syms_bool 
+syms_regs_set16(SymsRegs *regs, SymsRegID id, U16 value);
+
+REG_API syms_bool 
+syms_regs_set32(SymsRegs *regs, SymsRegID id, U32 value);
+
+REG_API syms_bool 
+syms_regs_set64(SymsRegs *regs, SymsRegID id, U64 value);
+
+REG_API syms_bool 
+syms_regs_set_addr(SymsRegs *regs, SymsRegID id, SymsAddr value);
+
+REG_API SymsAddr 
+syms_regs_get_ip(const SymsRegs *regs);
+
+REG_API SymsAddr 
+syms_regs_get_sp(const SymsRegs *regs);
+
+REG_API syms_bool 
+syms_regs_set_ip(SymsRegs *regs, SymsAddr value);
+
+REG_API syms_bool 
+syms_regs_set_sp(SymsRegs *regs, SymsAddr value);
+
+REG_API SymsEndian
+syms_regs_get_endianess(const SymsRegs *regs);
+
+#endif /* SYMS_REGS_INCLUDE_H */
diff --git a/Engine/Source/Runtime/SymsLib/include/regs/syms_regs_intel.h b/Engine/Source/Runtime/SymsLib/include/regs/syms_regs_intel.h
new file mode 100644
--- /dev/null
+++ b/Engine/Source/Runtime/SymsLib/include/regs/syms_regs_intel.h
@@ -0,0 +1,94 @@
+// Copyright Epic Games, Inc. All Rights Reserved.
+#ifndef SYMS_REGS_INTEL_H
+#define SYMS_REGS_INTEL_H
+
+// These FPU structs match the CPU's store layouts exactly.
+
+typedef struct SymsIntelFSave {
+  U16 fcw, _pad0;            // control word
+  U16 fsw, _pad1;            // status word
+  U16 ftw, _pad2;            // tag word
+  U32 fip;                   // EIP
+  U16 fcs;                   // CS
+  U16 fop;                   // last opcode
+  U32 fdp;                   // data pointer
+  U16 fds, _pad3;            // data selector
+  U8 st[80];                 // 8 * fpreg
+} SymsIntelFSave;
+SYMS_COMPILER_ASSERT(sizeof(SymsIntelFSave) == 108);
+
+typedef struct SymsIntelFXSave {
+  U16 fcw;                   // control word
+  U16 fsw;                   // status word
+  U16 ftw;                   // tag word
+  U16 fop;                   // last opcode
+  union {
+    struct {
+      U64 fip;           // RIP
+      U64 fdp;           // data pointer
+    } format64;
+    struct {
+      U32 fip;           // EIP
+      U16 fcs, _pad0;    // CS
+      U32 fdp;           // data pointer
+      U16 fds, _pad1;    // data selector
+    } format32;
+  } u;
+  U32 mxcsr;                 // MXCSR Register State
+  U32 mxcsr_mask;            // MXCSR Mask
+  U8 st[128];                // 8 * (fpreg + padding)
+  U8 xmm[256];               // 16 * xmmreg
+  U8 reserved[96];
+} SymsIntelFXSave;
+SYMS_COMPILER_ASSERT(sizeof(SymsIntelFXSave) == 512);
+
+typedef struct SymsIntelXSaveHdr {
+  U64 xstate_bv;
+  U8 reserved[56];
+} SymsIntelXSaveHdr;
+SYMS_COMPILER_ASSERT(sizeof(SymsIntelXSaveHdr) == 64);
+
+typedef struct SymsIntelXSave {
+  SymsIntelFXSave fxsave;
+  SymsIntelXSaveHdr xsave_hdr;
+
+  // NOTE:
+  // Technically the layout and size of the XSAVE struct is not defined
+  // beyond the header.
+  // The idea is that you should query CPUID to find out what features
+  // the CPU has, and the size+offset of each one.
+  // GDB/LLDB, and even the Linux kernel, assume that the YMMH registers
+  // are always at a fixed location. The Intel manuals suggest this to be
+  // the case too, but the AMD manuals don't.
+  // Anyway. This will do for now, in future we may need to query CPUID
+  // ourselves to know the layout here.
+  U8 ymmh[256];              // 16 * 16 bytes for each YMMH-reg
+} SymsIntelXSave;
+
+struct SymsRegs;
+
+SYMS_INTERNAL void 
+syms_regs_x64_fxsave_get_regs(struct SymsRegs *dest, SymsIntelFXSave *src);
+
+SYMS_INTERNAL void 
+syms_regs_x64_xsave_get_regs(struct SymsRegs *dest, SymsIntelXSave *src);
+
+SYMS_INTERNAL void 
+syms_regs_x86_fxsave_get_regs(struct SymsRegs *dest, SymsIntelFXSave *src);
+
+SYMS_INTERNAL void 
+syms_regs_x86_xsave_get_regs(struct SymsRegs *dest, SymsIntelXSave *src);
+
+SYMS_INTERNAL void 
+syms_regs_x64_fxsave_put_regs(SymsIntelFXSave *dest, struct SymsRegs *src);
+
+SYMS_INTERNAL void 
+syms_regs_x64_xsave_put_regs(SymsIntelXSave *dest, struct SymsRegs *src);
+
+SYMS_INTERNAL void 
+syms_regs_x86_fxsave_put_regs(SymsIntelFXSave *dest, struct SymsRegs *src);
+
+SYMS_INTERNAL void 
+syms_regs_x86_xsave_put_regs(SymsIntelXSave *dest, struct SymsRegs *src);
+
+#endif /* SYMS_REGS_INTEL_H */
diff --git a/Engine/Source/Runtime/SymsLib/include/syms.h b/Engine/Source/Runtime/SymsLib/include/syms.h
new file mode 100644
--- /dev/null
+++ b/Engine/Source/Runtime/SymsLib/include/syms.h
@@ -0,0 +1,110 @@
+// Copyright Epic Games, Inc. All Rights Reserved.
+#ifndef SYMS_INCLUDE_H
+#define SYMS_INCLUDE_H
+
+/******************************************************************************
+ * File   : syms.h                                                            *
+ * Author : Nikita Smith                                                      *
+ * Created: 2019/02/20                                                        *
+ * Purpose: Provides an API for accessing debug information                   *
+ ******************************************************************************/
+
+typedef struct SymsBorrowedMemory
+{
+  SymsArena *arena;
+  struct SymsBorrowedMemory *next;
+} SymsBorrowedMemory;
+
+struct SymsInstance
+{
+  syms_uint user_id;
+  SymsArena *arena;
+  SymsImage img;
+  SymsDebugInfo debug_info;
+  SymsMod null_mod;
+  SymsAddr rebase; // base address that is used for calculating addresses of symbols, lines, and etc.
+  syms_uint mods_num;
+  SymsModInfo *mods;
+  SymsBorrowedMemory *borrowed_memory;
+  struct 
+  {
+    SymsErrorCode code;
+    SymsString text;
+  } error;
+};
+
+
+typedef union SymsProcData
+{
+  struct 
+  {
+    pdb_cvdata_token cvdata;
+  } pdb;
+  struct 
+  {
+    DwTag type_tag;
+    DwAttrib frame_base;
+    SymsOffset range_off;
+  } dw;
+} SymsProcData;
+SYMS_COMPILER_ASSERT(sizeof(SymsProcData) <= sizeof(SymsProcImpl));
+
+typedef struct SymsMemread 
+{
+  SymsErrorCode result;
+  void *context;
+  syms_memread_sig *callback;
+} SymsMemread;
+
+typedef struct
+{
+  SymsErrorCode result;
+  SymsArch arch;
+  void *context;
+  syms_regread_sig *callback;
+} SymsRegread;
+
+typedef struct
+{
+  SymsErrorCode result;
+  SymsArch arch;
+  void *context;
+  syms_regwrite_sig *callback;
+} SymsRegwrite;
+
+SYMS_INTERNAL SymsTypeID
+syms_type_id_null(void);
+
+SYMS_INTERNAL void
+syms_line_init(SymsLine *line, SymsAddr va, syms_uint ln, syms_uint col);
+
+SYMS_API SymsModInfo *
+syms_get_mod(SymsInstance *instance, SymsModID mod_id);
+
+SYMS_API SymsModID
+syms_infer_global_data_module(SymsInstance *instance, SymsGlobal *gdata);
+// Resolve module that contains global-data.
+
+typedef struct SymsRangeMap
+{
+  SymsAddr lo;
+  SymsAddr hi;
+  syms_uint id;
+} SymsRangeMap;
+
+SYMS_API SymsRangeMap *
+syms_mod_info_find_rangemap(SymsModInfo *m, syms_uint i);
+
+SYMS_API SymsRangeMap *
+syms_rangemap_search(SymsBlockAllocator *rangemap, SymsAddr va);
+
+SYMS_INTERNAL SymsErrorCode
+syms_memread(SymsMemread* info, SymsAddr va, void* buffer, syms_uint buffer_size);
+SYMS_API SymsArenaFrame *
+syms_begin_arena_frame(struct SymsArena *arena);
+
+SYMS_API void
+syms_end_arena_frame(struct SymsArenaFrame *frame);
+
+#endif // SYMS_INCLUDE_H
+
diff --git a/Engine/Source/Runtime/SymsLib/include/syms_block_alloc.h b/Engine/Source/Runtime/SymsLib/include/syms_block_alloc.h
new file mode 100644
--- /dev/null
+++ b/Engine/Source/Runtime/SymsLib/include/syms_block_alloc.h
@@ -0,0 +1,40 @@
+// Copyright Epic Games, Inc. All Rights Reserved.
+#ifndef SYMS_BLOCK_ALLOC_H
+#define SYMS_BLOCK_ALLOC_H
+
+typedef struct SymsBlockNode
+{
+  char *base;
+  struct SymsBlockNode *next;
+} SymsBlockNode;
+
+struct SymsBlockAllocator
+{
+  syms_uint push_size;
+  syms_uint push_count;
+  syms_uint block_size;
+  syms_uint block_count;
+  syms_uint last_block_size;
+  SymsBlockNode *first_block;
+  SymsBlockNode **next_block;
+  SymsBlockNode *last_block;
+  SymsArena *arena;
+  SymsBlockNode **index_table;
+};
+
+SYMS_API SymsBlockAllocator *
+syms_block_allocator_init(syms_uint push_size, syms_uint block_size, SymsArena *arena);
+
+SYMS_API void *
+syms_block_allocator_push(SymsBlockAllocator *a);
+
+SYMS_API void
+syms_block_allocator_build_index_table(SymsBlockAllocator *a);
+
+SYMS_API void *
+syms_block_allocator_find_push(SymsBlockAllocator *a, syms_uint pi);
+
+SYMS_API void
+syms_block_allocator_sort(SymsBlockAllocator *a, syms_int (*compare)(const void *a, const void *b));
+
+#endif // SYMS_BLOCK_ALLOC_H
diff --git a/Engine/Source/Runtime/SymsLib/include/syms_core.h b/Engine/Source/Runtime/SymsLib/include/syms_core.h
new file mode 100644
--- /dev/null
+++ b/Engine/Source/Runtime/SymsLib/include/syms_core.h
@@ -0,0 +1,1284 @@
+// Copyright Epic Games, Inc. All Rights Reserved.
+#ifndef SYMS_CORE_INCLUDE_H
+#define SYMS_CORE_INCLUDE_H
+
+/******************************************************************************
+ * File   : syms_core.h                                                       *
+ * Author : Nikita Smith                                                      *
+ * Created: 2020/06/20                                                        *
+ * Purpose: core utilities for pdb and dwarf parsers                          *
+ ******************************************************************************/
+
+#define SYMS_PARANOID
+
+#ifndef syms_malloc
+#include <stdlib.h>
+#define syms_malloc(size, ud) ((void)ud,malloc(size))
+#define syms_free(ptr, ud)    ((void)ud,free(ptr))
+#endif
+#define syms_malloc_array(type, count, ud) (type *)syms_malloc(sizeof(type) * (count), ud)
+#define syms_malloc_struct(type, ud)       syms_malloc_array(type, 1, ud)
+
+#ifndef syms_qsort
+#include <stdlib.h>
+#define syms_qsort qsort
+#endif
+
+#ifndef syms_memcpy
+#include <string.h>
+#define syms_memcpy memcpy
+#endif
+
+#ifndef syms_memset
+#include <string.h>
+#define syms_memset memset
+#endif
+
+#ifndef syms_memcmp
+#include <string.h>
+#define syms_memcmp memcmp
+#endif
+
+#ifndef syms_strlen
+#include <string.h>
+#define syms_strlen strlen
+#endif
+
+#ifndef syms_strcmp
+#include <string.h>
+#define syms_strcmp strcmp
+#endif
+
+#ifndef syms_sprintf
+#include <stdio.h>
+#define syms_sprintf   sprintf
+#define syms_snprintf  snprintf
+#define syms_vsprintf  vsprintf
+#define syms_vsnprintf vsnprintf
+#endif
+
+#define syms_memzero(dst, size) syms_memset(dst, 0, size)
+
+#ifndef SYMS_INTERNAL
+#define SYMS_INTERNAL static
+#endif
+
+#ifndef SYMS_API
+#define SYMS_API
+#endif
+
+#ifndef SYMS_INLINE
+#define SYMS_INLINE inline
+#endif
+
+#ifndef SYMS_CUSTOM_TYPES
+#include <stdint.h>
+typedef uint8_t U8;
+typedef uint16_t U16;
+typedef uint32_t U32;
+typedef int8_t S8;
+typedef int16_t S16;
+typedef int32_t S32;
+typedef float F32;
+typedef double F64;
+typedef uint64_t U64;
+typedef int64_t S64;
+#endif
+
+typedef U8  u8;
+typedef U16 u16;
+typedef U32 u32;
+typedef U64 u64;
+typedef S8  s8;
+typedef S16 s16;
+typedef S32 s32;
+typedef S64 s64;
+typedef F32 f32;
+typedef F64 f64;
+
+typedef S32 syms_bool;
+
+#define syms_false 0
+#define syms_true 1
+
+#define SYMS_INT8_MIN  0x80 // -128
+#define SYMS_INT8_MAX  0x7f // +127
+#define SYMS_INT16_MIN 0x8000 // -32768
+#define SYMS_INT16_MAX 0x7fff // +32767
+#define SYMS_INT32_MIN 0x80000000 // -2147483648
+#define SYMS_INT32_MAX 0x7fffffff // +2147483647
+#define SYMS_INT64_MIN 0x8000000000000000ull // -9223372036854775808
+#define SYMS_INT64_MAX 0x7fffffffffffffffull // +9223372036854775807
+
+#define SYMS_UINT8_MAX  0xFFu
+#define SYMS_UINT16_MAX 0xFFFFu
+#define SYMS_UINT32_MAX 0xFFFFFFFFu
+#define SYMS_UINT64_MAX 0xFFFFFFFFFFFFFFFFu
+
+#define SYMS_TOSTR_X(x) #x
+#define SYMS_TOSTR(x) SYMS_TOSTR_X(x)
+#define SYMS_JOIN(a,b) a ## b
+
+#ifndef SYMS_ASSERT
+#ifdef _WIN32
+#include <intrin.h>
+#define SYMS_ASSERT(x) do { if (!(x)) __debugbreak(); } while (0)
+#else
+#include <assert.h>
+#define SYMS_ASSERT(x) assert(x)
+#endif
+#endif
+
+#ifndef SYMS_ASSERT_ALWAYS
+#define SYMS_ASSERT_ALWAYS(x) SYMS_ASSERT(x)
+#endif
+
+#ifndef SYMS_ASSERT_FAILURE
+#define SYMS_ASSERT_FAILURE(x) SYMS_ASSERT(0)
+#endif
+
+#ifndef SYMS_ASSERT_FAILURE_ALWAYS
+#define SYMS_ASSERT_FAILURE_ALWAYS(x) SYMS_ASSERT(0)
+#endif
+
+#if defined(SYMS_PARANOID)
+#define SYMS_ASSERT_PARANOID(x)           SYMS_ASSERT(x)
+#define SYMS_ASSERT_FAILURE_PARANOID(msg) SYMS_ASSERT_FAILURE(msg)
+#else
+#define SYMS_ASSERT_PARANOID(x)
+#define SYMS_ASSERT_FAILURE_PARANOID(msg)
+#endif
+
+#define SYMS_NOT_IMPLEMENTED    SYMS_ASSERT_FAILURE_ALWAYS("not implemented")
+#define SYMS_INVALID_CODE_PATH  SYMS_ASSERT_FAILURE_ALWAYS("invalid code path")
+#define SYMS_ASSERT_NO_SUPPORT  SYMS_INVALID_CODE_PATH
+
+#ifndef SYMS_COMPILER_ASSERT
+#define SYMS_COMPILER_ASSERT(predicate) SYMS_COMPILER_ASSERT_IMPL(predicate, __LINE__)
+#define SYMS_COMPILER_ASSERT_IMPL(predicate, line) typedef char SYMS_JOIN(compiler_dummy_name, line)[(predicate) ? +1 : -1]
+#endif
+
+#ifndef SYMS_ARRAY_SIZE
+#define SYMS_ARRAY_SIZE(x) (sizeof(x)/sizeof((x)[0]))
+#endif
+
+#ifndef SYMS_MIN
+#define SYMS_MIN(a,b) ((a) < (b) ? (a) : (b))
+#endif
+
+#ifndef SYMS_MAX
+#define SYMS_MAX(a,b) ((a) > (b) ? (a) : (b))
+#endif
+
+#ifndef SYMS_ABS
+#define SYMS_ABS(x) ((x) < 0 ? -(x) : (x))
+#endif
+
+#ifndef SYMS_MEMBER_OFFSET
+#define SYMS_MEMBER_OFFSET(type, member) offsetof(type, member)
+#endif
+
+#ifndef SYMS_GET_MEMBER_SIZE
+#define SYMS_GET_MEMBER_SIZE(type, member) sizeof((((type*)0)->member))
+#endif
+
+#ifdef _MSC_VER
+#define SYMS_PRIu64 "I64u"
+#define SYMS_PRId64 "I64d"
+#define SYMS_PRIx64 "I64x"
+#else
+#include <inttypes.h>
+#define SYMS_PRIu64 PRIu64
+#define SYMS_PRId64 PRId64
+#define SYMS_PRIx64 PRIx64
+#endif
+
+#define SYMS_KB(num) ((num)*1024u)
+#define SYMS_MB(num) (SYMS_KB(num)*1024u)
+#define SYMS_GB(num) (SYMS_MB(num)*1024u)
+#define SYMS_TB(num) (SYMS_TB(num)*1024u)
+
+#if defined(__amd64__) || defined(__amd64) || defined(__x86_64__) || defined(_M_X64) || defined(_M_AMD64)
+#define SYMS_X86
+#define SYMS_X64
+#define SYMS_LITTLE_ENDIAN
+#elif defined(i386) || defined(__i386) || defined(__i386__) || defined(_X86_) || defined(_M_I86)
+#define SYMS_X86
+#define SYMS_LITTLE_ENDIAN
+#elif defined(__arm__) || defined(_M_ARM) 
+#define SYMS_ARM
+#define SYMS_LITTLE_ENDIAN
+#elif defined(__aarch64__)
+#define SYMS_ARM64
+#define SYMS_LITTLE_ENDIAN
+#else
+#error "arch not defined"
+#endif
+
+#ifndef SYMS_CUSTOM_BSWAP
+#if defined(__clang__)
+#define syms_bswap16 __builtin_bswap16
+#define syms_bswap32 __builtin_bswap32
+#define syms_bswap64 __builtin_bswap64
+#elif defined(_MSC_VER)
+#define syms_bswap16 _byteswap_ushort
+#define syms_bswap32 _byteswap_ulong
+#define syms_bswap64 _byteswap_uint64
+#else
+SYMS_INTERNAL U16 syms_bswap16(U16 x) 
+{
+    return (((x >> 8) & 0xffu) | ((x & 0xffu) << 8));
+}
+SYMS_INTERNAL U32 syms_swap32(U32 x) 
+{
+    return (((x & 0xff000000u) >> 24) |
+            ((x & 0x00ff0000u) >> 8) |
+            ((x & 0x0000ff00u) << 8) |
+            ((x & 0x000000ffu) << 24));
+}
+SYMS_INTERNAL U64 syms_bswap64(U64 x) 
+{
+    return (((x & 0xff00000000000000) >> 56) |
+            ((x & 0x00ff000000000000) >> 40) |
+            ((x & 0x0000ff0000000000) >> 24) |
+            ((x & 0x000000ff00000000) >> 8) |
+            ((x & 0x00000000ff000000) << 8) |
+            ((x & 0x0000000000ff0000) << 24) |
+            ((x & 0x000000000000ff00) << 40) |
+            ((x & 0x00000000000000ff) << 56));
+}
+#endif
+#endif
+
+#if defined(SYMS_X64) || defined(SYMS_ARM64)
+typedef U64 SymsUWord;
+typedef S64 SymsSWord;
+#define SYMS_UWORD_MAX SYMS_UINT64_MAX
+#define SYMS_SWORD_MAX SYMS_INT64_MAX
+#define syms_bswapaddr(x) syms_bswap64(x)
+#elif defined(SYMS_X86)
+typedef U32 SymsUWord;
+typedef S32 SymsSWord;
+#define SYMS_UWORD_MAX SYMS_UINT32_MAX
+#define SYMS_SWORD_MAX SYMS_INT32_MAX
+#define syms_bswapaddr(x) syms_bswap32(x)
+#else
+#error "Address is not defined for this architecture"
+#endif
+
+//typedef SymsUMM SymsAddr;
+
+typedef enum 
+{
+    SYMS_ADDR_SIZE_NULL,
+    SYMS_ADDR_SIZE_16,
+    SYMS_ADDR_SIZE_32,
+    SYMS_ADDR_SIZE_64,
+    SYMS_ADDR_SIZE_128,
+    
+    SYMS_ADDR_SIZE_HOST =
+#if SYMS_ADDR_MAX == SYMS_UINT16_MAX
+    SYMS_ADDR_SIZE_16
+#elif SYMS_ADDR_MAX == SYMS_UINT32_MAX
+    SYMS_ADDR_SIZE_32
+#elif SYMS_ADDR_MAX == SYMS_UINT64_MAX 
+    SYMS_ADDR_SIZE_64
+#elif SYMS_ADDR_MAX == SYMS_UINT64_MAX*2
+    SYMS_ADDR_SIZE_128
+#else
+#error "Address size is not defined"
+#endif
+} SymsAddrSize;
+
+typedef enum 
+{
+    SYMS_ENDIAN_NULL,
+    SYMS_ENDIAN_LITTLE,
+    SYMS_ENDIAN_BIG,
+    
+    SYMS_ENDIAN_HOST = SYMS_ENDIAN_LITTLE
+} SymsEndian;
+
+#if defined(SYMS_LITTLE_ENDIAN)
+#define SYMS_LE16(x)     (x)
+#define SYMS_LE32(x)     (x)
+#define SYMS_LE64(x)     (x)
+#define SYMS_LEADDR(x)   (x)
+#define SYMS_BE16(x)     syms_bswap16(x)
+#define SYMS_BE32(x)     syms_bswap32(x)
+#define SYMS_BE64(x)     syms_bswap64(x)
+#define SYMS_BEADDR(x)   syms_bswapaddr(x)
+#else
+#define SYMS_LE16(x)     syms_bswap16(x)
+#define SYMS_LE32(x)     syms_bswap32(x)
+#define SYMS_LE64(x)     syms_bswap64(x)
+#define SYMS_LEADDR(x)   syms_bswapaddr(x)
+#define SYMS_BE16(x)     (x)
+#define SYMS_BE32(x)     (x)
+#define SYMS_BE64(x)     (x)
+#define SYMS_BEADDR(x)   (x)
+#endif
+
+#define SYMS_OFFSET_INVALID SYMS_UWORD_MAX
+typedef SymsUWord SymsOffset;
+
+// -------------------------------------------------------------------------------- 
+
+typedef struct SymsBlock
+{
+    SymsUMM len;
+    void *data;
+} SymsBlock;
+
+// -------------------------------------------------------------------------------- 
+
+#define SYMS_STRING_LIT(x) { sizeof(x) - 1, x }
+#define syms_string_lit(x) syms_string_init(x, sizeof(x) - 1)
+
+typedef struct SymsStringArray 
+{
+    U32 count;
+    SymsString *entries;
+} SymsStringArray;
+
+SYMS_API SymsString
+syms_string_init(const char *data, U32 len);
+
+SYMS_API SymsString
+syms_string_init_lit(const char *data);
+
+SYMS_API syms_bool
+syms_string_is_null(SymsString str);
+
+SYMS_API syms_bool
+syms_string_cmp_lit(SymsString a, const char *b);
+
+SYMS_API syms_bool
+syms_is_alpha_ascii(U32 codepoint);
+
+SYMS_API U32
+syms_lowercase(U32 codepoint);
+
+SYMS_API U32
+syms_uppercase(U32 codepoint);
+
+SYMS_API syms_bool
+syms_string_cmp_nocase(SymsString a_str, SymsString b_str);
+
+SYMS_API char
+syms_string_peek_byte(SymsString str, U32 byte_offset);
+
+SYMS_API U32
+syms_string_get_size(SymsString str);
+
+SYMS_API syms_bool
+syms_string_cmp(SymsString a, SymsString b);
+
+SYMS_API const U8 *
+syms_decode_utf8(const U8 *p, U32 *dst);
+
+SYMS_API char *
+syms_string_to_cstr(SymsString str, SymsArena *arena);
+
+
+// -------------------------------------------------------------------------------- 
+
+typedef struct SymsBuffer 
+{
+    const void *base;
+    SymsUMM size;
+    SymsUMM off;
+    u32 addr_size;
+    u32 offs_size;
+} SymsBuffer;
+
+SYMS_API SymsBuffer
+syms_buffer_init(const void *data, SymsUMM size);
+
+SYMS_API SymsBuffer
+syms_buffer_init_ex(const void *data, SymsUMM size, SymsAddrSize addr_size, SymsAddrSize offset_type);
+
+SYMS_API syms_bool
+syms_buffer_skip(SymsBuffer *buffer, SymsOffset num);
+
+SYMS_API syms_bool
+syms_buffer_seek(SymsBuffer *buffer, SymsOffset new_off);
+
+SYMS_API U8
+syms_buffer_peek_u8(SymsBuffer *bin);
+
+SYMS_API U16
+syms_buffer_peek_u16(SymsBuffer *bin);
+
+SYMS_API U32
+syms_buffer_peek_u32(SymsBuffer *bin);
+
+SYMS_API U64
+syms_buffer_peek_u64(SymsBuffer *bin);
+
+SYMS_API S8
+syms_buffer_peek_s8(SymsBuffer *bin);
+
+SYMS_API S16
+syms_buffer_peek_s16(SymsBuffer *bin);
+
+SYMS_API S32
+syms_buffer_peek_s32(SymsBuffer *bin);
+
+SYMS_API S64
+syms_buffer_peek_s64(SymsBuffer *bin);
+
+SYMS_API S64
+syms_buffer_peek_sleb128(SymsBuffer *bin);
+
+SYMS_API U64
+syms_buffer_peek_uleb128(SymsBuffer *bin);
+
+SYMS_API SymsString
+syms_buffer_peek_string(SymsBuffer *bin);
+
+SYMS_API SymsString
+syms_buffer_peek_cstr(SymsBuffer *bin);
+
+#define syms_buffer_push_array(buffer, type, count) (type *)syms_buffer_push(buffer, sizeof(type) * (count))
+#define syms_buffer_push_struct(buffer, type)       syms_buffer_push_array(buffer, type, 1)
+SYMS_API void *
+syms_buffer_push(SymsBuffer *buffer, SymsUMM num_bytes);
+
+SYMS_API syms_bool
+syms_buffer_read(SymsBuffer *buffer, void *out_bf, SymsUMM size);
+
+SYMS_API U8
+syms_buffer_read_u8(SymsBuffer *buffer);
+
+SYMS_API U16
+syms_buffer_read_u16(SymsBuffer *buffer);
+
+SYMS_API U32
+syms_buffer_read_u32(SymsBuffer *buffer);
+
+SYMS_API U32
+syms_buffer_read_u32(SymsBuffer *buffer);
+
+SYMS_API U64
+syms_buffer_read_u64(SymsBuffer *buffer);
+
+SYMS_API U64
+syms_buffer_read_uleb128(SymsBuffer *bin);
+// Reads an unsigned integer with variable-length encoding LEB128.
+
+SYMS_API S8
+syms_buffer_read_s8(SymsBuffer *buffer);
+
+SYMS_API S16
+syms_buffer_read_s16(SymsBuffer *buffer);
+
+SYMS_API S32
+syms_buffer_read_s32(SymsBuffer *buffer);
+
+SYMS_API S64
+syms_buffer_read_s64(SymsBuffer *buffer);
+
+SYMS_API S64
+syms_buffer_read_sleb128(SymsBuffer *bin);
+
+SYMS_API SymsString
+syms_buffer_read_string(SymsBuffer *buffer);
+
+SYMS_API SymsString
+syms_buffer_read_cstr(SymsBuffer *buffer);
+
+SYMS_API void
+syms_buffer_write_cstr(SymsBuffer *buffer, const char *cstr);
+
+SYMS_API void
+syms_buffer_write(SymsBuffer *buffer, const void *src, U32 src_size);
+
+SYMS_API void
+syms_buffer_write_u8(SymsBuffer *buffer, U8 value);
+
+SYMS_API const char *
+syms_buffer_get_cstr(SymsBuffer *buffer);
+
+SYMS_API const void *
+syms_buffer_at(SymsBuffer *buffer);
+
+// -------------------------------------------------------------------------------- 
+
+SYMS_API SymsString
+syms_path_get_file_name(SymsString path, syms_bool strip_ext);
+
+// -------------------------------------------------------------------------------- 
+
+enum 
+{
+    SYMS_ARENA_FLAG_ALLOC_FAILED = (1 << 0)
+};
+typedef u32 SymsArenaFlags;
+
+struct SymsArena 
+{
+    SymsArenaFlags flags;
+    SymsUMM size;
+    SymsUMM page_size;
+    void *head;
+    u8 *cursor;
+};
+
+struct SymsArenaFrame
+{
+    SymsArena *arena;
+    SymsArenaFlags flags;
+    SymsUMM size;
+    void *head;
+    u8 *cursor;
+};
+
+SYMS_API void
+syms_arena_init(SymsArena *arena, SymsUMM page_size);
+
+SYMS_API void
+syms_arena_free(SymsArena *arena);
+
+#define syms_arena_push_array_ex(arena, type, num, align) (type *)syms_arena_push_ex(arena, sizeof(type)*num, align)
+#define syms_arena_push_struct_ex(arena, type, align) syms_arena_push_array_ex(arena, type, 1, align)
+
+#define syms_arena_push(arena, len)             syms_arena_push_ex(arena, len, sizeof(void*))
+#define syms_arena_push_array(arena, type, num) syms_arena_push_array_ex(arena, type, num, sizeof(void*))
+#define syms_arena_push_struct(arena, type)     syms_arena_push_array(arena, type, 1)
+
+SYMS_API void *
+syms_arena_push_ex(SymsArena *arena, SymsUMM len, SymsUMM align);
+
+SYMS_API char *
+syms_arena_push_cstr(SymsArena *arena, const char *cstr);
+
+SYMS_API SymsArenaFrame
+syms_arena_frame_begin(SymsArena *arena);
+
+SYMS_API void
+syms_arena_frame_end(SymsArenaFrame frame);
+
+// -------------------------------------------------------------------------------- 
+
+SYMS_API U16
+syms_trunc_u16(U32 value);
+
+SYMS_API U32
+syms_trunc_u32(U64 value);
+
+U32
+syms_hash_djb2(const char *str, u32 size);
+
+U16
+syms_hash_djb2_16(const char *str, u32 size);
+
+#endif /* SYMS_CORE_INCLUDE_H */
+
+#if defined(SYMS_CORE_IMPLEMENT)
+
+SYMS_API U32
+syms_addrsize_to_int(SymsAddrSize addr)
+{
+    U32 result;
+    switch (addr) {
+        default:
+        case SYMS_ADDR_SIZE_NULL: result = 0; break;
+        case SYMS_ADDR_SIZE_16:   result = 2; break;
+        case SYMS_ADDR_SIZE_32:   result = 4; break;
+        case SYMS_ADDR_SIZE_64:   result = 8; break;
+        case SYMS_ADDR_SIZE_128:  result = 16; break;
+    }
+    return result;
+}
+
+SYMS_API U16
+syms_trunc_u16(U32 value)
+{
+    U16 result;
+    SYMS_ASSERT_ALWAYS(value <= SYMS_UINT16_MAX);
+    result = (U16)value;
+    return result;
+}
+
+SYMS_API U32
+syms_trunc_u32(U64 value)
+{
+    U32 result;
+    SYMS_ASSERT_ALWAYS(value <= SYMS_UINT32_MAX);
+    result = (U32)value;
+    return result;
+}
+
+SYMS_API void
+syms_arena_init(SymsArena *a, SymsUMM page_size)
+{
+  a->flags = 0;
+  a->size = 0;
+  a->cursor = 0;
+  a->page_size = page_size;
+  a->head = 0;
+}
+
+SYMS_API void
+syms_arena_free(SymsArena *a)
+{
+  void *p = a->head;
+  while (p) {
+    void *n = *(void **)p;
+    *(void **)p = 0;
+    syms_free_virtual_memory(p, 0);
+    p = n;
+  }
+}
+
+SYMS_API void *
+syms_arena_push_ex(SymsArena *a, SymsUMM size, SymsUMM align)
+{
+  SymsUMM cursor = (SymsUMM)a->cursor;
+  SymsUMM mask = align - 1;
+  SymsUMM pad = 0;
+  if (cursor & mask) {
+    pad = align - (cursor & mask);
+  }
+  SymsUMM total_size = size + pad;
+  SymsUMM cursor_offset = cursor - (SymsUMM)a->head;
+  if (cursor_offset + total_size > a->size) {
+    SymsUMM pc = (total_size + 16 + (a->page_size - 1)) / a->page_size;
+    SymsUMM as = pc * a->page_size;
+    u8 *p = (u8 *)syms_reserve_virtual_memory(as);
+    SYMS_ASSERT(p);
+    syms_commit_virtual_memory(p, as);
+    ((void**)p)[0] = a->head;
+    a->head = p;
+    a->cursor = p + 16;
+    a->size = pc * a->page_size;
+  }
+  void *ptr = a->cursor + pad;
+  a->cursor += total_size;
+  return ptr;
+}
+SYMS_API char *
+syms_arena_push_cstr(SymsArena *arena, const char *cstr)
+{
+    u32 cstr_size = (u32)syms_strlen(cstr);
+    char *result = syms_arena_push_array(arena, char, cstr_size);
+    syms_memcpy(result, cstr, cstr_size);
+    return result;
+}
+
+SYMS_API SymsArenaFrame
+syms_arena_frame_begin(SymsArena *a)
+{
+    SymsArenaFrame f;
+    f.arena = a;
+    f.flags = a->flags;
+    f.cursor = a->cursor;
+    f.head = a->head;
+    f.size = a->size;
+    return f;
+}
+
+SYMS_API void
+syms_arena_frame_end(SymsArenaFrame f)
+{
+  SymsArena *a = f.arena;
+  void *p = a->head;
+  while (p != f.head) {
+    void *n = *(void **)p;
+    *(void **)p = 0;
+    syms_free_virtual_memory(p, 0);
+    p = n;
+  }
+  a->head = f.head;
+  a->cursor = f.cursor;
+  a->flags = f.flags;
+  a->size = f.size;
+}
+
+SYMS_API SymsString
+syms_string_init(const char *data, U32 len)
+{
+    SymsString result;
+    while (len > 0 && data[len - 1] == '\0') {
+        --len;
+    }
+    result.data = data;
+    result.len = len;
+    return result;
+}
+
+SYMS_API SymsString
+syms_string_init_lit(const char *data)
+{
+    SymsString result;
+    result.len = 0;
+    result.data = data;
+    while (*data++) {
+        ++result.len;
+    }
+    return result;
+}
+SYMS_API syms_bool
+syms_string_cmp_lit(SymsString a, const char *b)
+{
+    for (;;) {
+        if (a.len == 0 || *b == '\0')
+            break;
+        
+        if (*a.data++ != *b++)
+            break;
+        
+        --a.len;
+    }
+    
+    return a.len == 0 && *b == '\0';
+}
+SYMS_API U32
+syms_string_get_size(SymsString str)
+{
+    SYMS_ASSERT(str.data);
+    return str.len;
+}
+SYMS_API syms_bool
+syms_string_is_null(SymsString str)
+{
+    syms_bool result = str.data == 0 || str.len == 0;
+    return result;
+}
+SYMS_API syms_bool
+syms_string_cmp(SymsString a, SymsString b)
+{
+    syms_bool result = syms_false;
+    if (a.len == b.len) {
+        U32 i;
+        for (i = 0; i < a.len; ++i) {
+            if (a.data[i] != b.data[i])
+                break;
+        }
+        result = i == a.len;
+    }
+    return result;
+}
+SYMS_API char *
+syms_string_to_cstr(SymsString str, SymsArena *arena)
+{
+    char *p = syms_arena_push_array(arena, char, str.len + 1);
+    if (p) {
+        syms_memcpy(p, str.data, str.len);
+        p[str.len] = '\0';
+    }
+    return p;
+}
+
+SYMS_API syms_bool
+syms_is_alpha_ascii(U32 codepoint)
+{
+    return (codepoint >= 'A' && codepoint <= 'Z') || (codepoint >= 'a' && codepoint <= 'z');
+}
+SYMS_API U32
+syms_lowercase(U32 codepoint)
+{
+    if (codepoint >= 'A' && codepoint <= 'Z')
+        codepoint += 'a' - 'A';
+    return codepoint;
+}
+SYMS_API U32
+syms_uppercase(U32 codepoint)
+{
+    if (codepoint >= 'a' && codepoint <= 'z')
+        codepoint -= 'a' - 'A';
+    return codepoint;
+}
+SYMS_API syms_bool
+syms_string_cmp_nocase(SymsString a_str, SymsString b_str)
+{
+    syms_bool result = syms_false;
+    if (a_str.len == b_str.len) {
+        U32 i;
+        for (i = 0; i < a_str.len; ++i) {
+            U32 a = syms_lowercase((U32)a_str.data[i]);
+            U32 b = syms_lowercase((U32)b_str.data[i]);
+            if (a != b) break;
+        }
+        result = (i == a_str.len);
+    }
+    return result;
+}
+SYMS_API char
+syms_string_peek_byte(SymsString str, U32 byte_offset)
+{
+    char result = 0;
+    if (byte_offset < str.len) {
+        result = ((const char *)str.data)[byte_offset];
+    }
+    return result;
+}
+
+SYMS_API const U8 *
+syms_decode_utf8(const U8 *p, U32 *dst)
+{
+    U32 res, n;
+    switch (*p & 0xf0) {
+        case 0xf0 :  res = *p & 0x07;  n = 3;  break;
+        case 0xe0 :  res = *p & 0x0f;  n = 2;  break;
+        case 0xd0 :
+        case 0xc0 :  res = *p & 0x1f;  n = 1;  break;
+        default   :  res = *p;         n = 0;  break;
+    }
+    while (n--) {
+        res = (res << 6) | (*(++p) & 0x3f);
+    }
+    *dst = res;
+    return p + 1;
+}
+
+SYMS_API SymsString
+syms_path_get_file_name(SymsString path, syms_bool strip_ext)
+{
+    const U8 *p = (const U8 *)path.data;
+    const U8 *p_end = p + path.len;
+    const U8 *start = p;
+    const U8 *end = p_end;
+    u32 fname_size;
+    SymsString fname;
+    
+    while (p != p_end) {
+        U32 c;
+        p = syms_decode_utf8(p, &c);
+        if (c == '\0') {
+            break;
+        }
+        if (c == '\\' || c == '/') {
+            start = p;
+        }
+    }
+    if (strip_ext) {
+        p = start;
+        p_end = end;
+        while (p != p_end) {
+            U32 c;
+            p = syms_decode_utf8(p, &c);
+            if (c == '.') {
+                end = p;
+            }
+        }
+    }
+    
+    fname_size = (U32)(end - start);
+    fname = syms_string_init((const char *)start, fname_size);
+    return fname;
+}
+
+SYMS_API syms_bool
+syms_buffer_read(SymsBuffer *bin, void *out_bf, SymsUMM size)
+{
+    syms_bool done = syms_false;
+    if (bin->off + size <= bin->size) {
+        const void *src = (void *)((const char *)bin->base + bin->off);
+        syms_memcpy(out_bf, src, size);
+        bin->off += size;
+        done = syms_true;
+    }
+    return done;
+}
+
+SYMS_API void *
+syms_buffer_push(SymsBuffer *bin, SymsUMM size)
+{
+    void *result = 0;
+    if (bin->off + size <= bin->size) {
+        if (bin->base) {
+            result = (void *)((const U8 *)bin->base + bin->off);
+        }
+        bin->off += size;
+    }
+    return result;
+}
+
+SYMS_API SymsBuffer
+syms_buffer_init_ex(const void *base, SymsUMM size, SymsAddrSize addr_type, SymsAddrSize offs_type)
+{
+    SymsBuffer result;
+    
+    result.base      = base;
+    result.size      = size;
+    result.off       = 0;
+    result.addr_size = syms_addrsize_to_int(addr_type);
+    result.offs_size = syms_addrsize_to_int(offs_type);
+    
+#if defined(SYMS_PARANOID)
+    SYMS_ASSERT(result.addr_size > 0);
+    SYMS_ASSERT(result.offs_size > 0);
+#endif
+    
+    return result;
+}
+
+SYMS_API  SymsBuffer
+syms_buffer_init(const void *data, SymsUMM size)
+{
+    return syms_buffer_init_ex(data, size, SYMS_ADDR_SIZE_HOST, SYMS_ADDR_SIZE_HOST);
+}
+
+SYMS_API syms_bool
+syms_buffer_skip(SymsBuffer *bin, SymsOffset num)
+{
+    syms_bool result = syms_false;
+    if (bin->off + num <= bin->size) {
+        bin->off += num;
+        result = syms_true;
+    }
+    return result;
+}
+
+SYMS_API syms_bool
+syms_buffer_seek(SymsBuffer *bin, SymsOffset new_off)
+{
+    syms_bool result = syms_false;
+    if (new_off <= bin->size) {
+        bin->off = new_off;
+        result = syms_true;
+    }
+    return result;
+}
+
+#define syms_buffer_peek(bin, a) syms_buffer_peek_ex(bin, a, a)
+#define syms_buffer_peek_ex(bin, a, b) \
+do { \
+SymsOffset off = bin->off; \
+a result = syms_buffer_read_##b(bin); \
+bin->off = off; \
+return result; \
+} while (0)
+
+
+SYMS_API U8
+syms_buffer_peek_u8(SymsBuffer *bin)
+{
+    syms_buffer_peek(bin, u8);
+}
+
+SYMS_API U16
+syms_buffer_peek_u16(SymsBuffer *bin)
+{
+    syms_buffer_peek(bin, u16);
+}
+
+SYMS_API U32
+syms_buffer_peek_u32(SymsBuffer *bin)
+{
+    syms_buffer_peek(bin, u32);
+}
+
+SYMS_API U64
+syms_buffer_peek_u64(SymsBuffer *bin)
+{
+    syms_buffer_peek(bin, u64);
+}
+
+SYMS_API S8
+syms_buffer_peek_s8(SymsBuffer *bin)
+{
+    syms_buffer_peek(bin, s8);
+}
+
+SYMS_API S16
+syms_buffer_peek_s16(SymsBuffer *bin)
+{
+    syms_buffer_peek(bin, s16);
+}
+
+SYMS_API S32
+syms_buffer_peek_s32(SymsBuffer *bin)
+{
+    syms_buffer_peek(bin, s32);
+}
+
+SYMS_API S64
+syms_buffer_peek_s64(SymsBuffer *bin)
+{
+    syms_buffer_peek(bin, s64);
+}
+
+SYMS_API S64
+syms_buffer_peek_sleb128(SymsBuffer *bin)
+{
+    syms_buffer_peek_ex(bin, s64, sleb128);
+}
+
+SYMS_API U64
+syms_buffer_peek_uleb128(SymsBuffer *bin)
+{
+    syms_buffer_peek_ex(bin, u64, uleb128);
+}
+
+SYMS_API SymsString
+syms_buffer_peek_string(SymsBuffer *bin)
+{
+    syms_buffer_peek_ex(bin, SymsString, string);
+}
+
+SYMS_API SymsString
+syms_buffer_peek_cstr(SymsBuffer *bin)
+{
+    syms_buffer_peek_ex(bin, SymsString, cstr);
+}
+
+SYMS_API void
+syms_buffer_write(SymsBuffer *bin, const void *src, U32 src_size)
+{
+    void *dst = syms_buffer_push(bin, src_size);
+    if (dst) {
+        syms_memcpy(dst, src, src_size);
+    }
+}
+SYMS_API void
+syms_buffer_write_cstr(SymsBuffer *bin, const char *cstr)
+{
+    U32 cstr_size = (U32)syms_strlen(cstr);
+    syms_buffer_write(bin, cstr, cstr_size);
+}
+SYMS_API const char *
+syms_buffer_get_cstr(SymsBuffer *buffer)
+{
+    char *base = (char *)buffer->base;
+    if (buffer->off >= buffer->size) {
+        SYMS_ASSERT(buffer->size > 0);
+        syms_memset((char*)base + buffer->size - 1, '\0', 1);
+    } else {
+        syms_memset((char*)base + buffer->off, '\0', 1);
+    }
+    return base;
+}
+SYMS_API const void *
+syms_buffer_at(SymsBuffer *buffer)
+{
+    const void *result = (const void *)((const char *)buffer->base + buffer->off);
+    return result;
+}
+
+SYMS_API void
+syms_buffer_write_u8(SymsBuffer *buffer, U8 value)
+{
+    syms_buffer_write(buffer, &value, sizeof(value));
+}
+
+SYMS_API void
+syms_buffer_write_u32(SymsBuffer *buffer, U32 value)
+{
+    syms_buffer_write(buffer, &value, sizeof(value));
+}
+
+SYMS_API SymsString
+syms_buffer_read_cstr(SymsBuffer *buffer)
+{
+    const char *p = (const char *)buffer->base + buffer->off;
+    SymsUWord off = buffer->off;
+    u32 size;
+    SymsString result;
+    while (off < buffer->size) {
+        if (p[off] == '\0') {
+            break;
+        }
+        off += 1;
+    }
+    size = syms_trunc_u32(off - buffer->off);
+    buffer->off = off;
+    result = syms_string_init(p, size);
+    return result;
+}
+
+SYMS_API SymsString
+syms_buffer_read_string(SymsBuffer *buffer)
+{
+    SymsUWord off_start = buffer->off;
+    SymsUWord off_end = buffer->off;
+    const char *ptr = (const char *)buffer->base + buffer->off;
+    SymsString result;
+    while (buffer->off < buffer->size) {
+        if (((const char *)buffer->base)[buffer->off++] == '\0')
+            break;
+    }
+    off_end = buffer->off;
+    result = syms_string_init(ptr, (U32)(off_end - off_start));
+    return result;
+}
+
+SYMS_API U8
+syms_buffer_read_u8(SymsBuffer *buffer)
+{
+    U8 result = 0;
+    syms_buffer_read(buffer, &result, sizeof(result));
+    return result;
+}
+
+SYMS_API U16
+syms_buffer_read_u16(SymsBuffer *buffer)
+{
+    U16 result = 0;
+    syms_buffer_read(buffer, &result, sizeof(result));
+    return result;
+}
+
+SYMS_API U32
+syms_buffer_read_u24(SymsBuffer *buffer)
+{
+    U32 result;
+    U32 a = (U32)syms_buffer_read_u8(buffer);
+    U32 b = (U32)syms_buffer_read_u8(buffer);
+    U32 c = (U32)syms_buffer_read_u8(buffer);
+    result = (a << 0) | (b << 8) | (c << 16);
+    return result;
+}
+
+SYMS_API U32
+syms_buffer_read_u32(SymsBuffer *buffer)
+{
+    U32 result = 0;
+    syms_buffer_read(buffer, &result, sizeof(result));
+    return result;
+}
+
+SYMS_API U64
+syms_buffer_read_u64(SymsBuffer *buffer)
+{
+    U64 result = 0;
+    syms_buffer_read(buffer, &result, sizeof(result));
+    return result;
+}
+
+SYMS_API S8
+syms_buffer_read_s8(SymsBuffer *buffer)
+{
+    S8 result = 0;
+    syms_buffer_read(buffer, &result, sizeof(result));
+    return result;
+}
+
+SYMS_API S16
+syms_buffer_read_s16(SymsBuffer *buffer)
+{
+    S16 result = 0;
+    syms_buffer_read(buffer, &result, sizeof(result));
+    return result;
+}
+
+SYMS_API S32
+syms_buffer_read_s32(SymsBuffer *buffer)
+{
+    S32 result = 0;
+    syms_buffer_read(buffer, &result, sizeof(result));
+    return result;
+}
+
+SYMS_API S64
+syms_buffer_read_s64(SymsBuffer *buffer)
+{
+    S64 result = 0;
+    syms_buffer_read(buffer, &result, sizeof(result));
+    return result;
+}
+
+SYMS_API U64
+syms_buffer_read_uleb128(SymsBuffer *buffer)
+{
+    U64 res = 0;
+    U64 shift = 0;
+    
+    const U8 *data = (const U8 *)buffer->base + buffer->off;
+    const U8 *start = data;
+    
+    U32 num_used;
+    
+    while (buffer->off < buffer->size) {
+        U8 i = *data++;
+        U8 val = i & 0x7fu;
+        res |= (U64)val << shift;
+        if ((i & 0x80u) == 0) {
+            break;
+        }
+        shift += 7u;
+    }
+    
+    num_used = (U32)((const U8 *)data - (const U8 *)start);
+    buffer->off += num_used;
+    
+    return res;
+}
+
+SYMS_API S64
+syms_buffer_read_sleb128(SymsBuffer *buffer)
+{
+    S64 result = 0;
+    U32 shift = 0;
+    
+    const U8 *data = (const U8 *)buffer->base + buffer->off;
+    const U8 *end = (const U8 *)buffer->base + buffer->size;
+    const U8 *start = data;
+    
+    while (data < end) {
+        U8 byte = *data;
+        data += 1;
+        
+        result |= ((U64)(byte & 0x7f)) << shift;
+        shift += 7;
+        
+        if ((byte & 0x80) == 0) {
+            U32 num_used;
+            
+            if (shift < (sizeof(result) * 8) && (byte & 0x40)) {
+                result |= -((S64)(1u << shift));
+            }
+            
+            num_used = (U32)((const U8 *)data - (const U8 *)start);
+            buffer->off += num_used;
+            
+            return (S64)(result);
+        }
+    }
+    
+    return 0;
+}
+
+SYMS_API SymsAddr
+syms_buffer_read_addr(SymsBuffer *buffer)
+{
+    SymsAddr result = 0;
+    SYMS_ASSERT(buffer->addr_size > 0);
+    SYMS_ASSERT(buffer->addr_size <= sizeof(SymsAddr));
+    syms_buffer_read(buffer, &result, buffer->addr_size);
+    return result;
+}
+
+SYMS_API SymsAddr
+syms_buffer_read_offs(SymsBuffer *buffer)
+{
+    SymsAddr result = 0;
+    SYMS_ASSERT(buffer->offs_size > 0);
+    SYMS_ASSERT(buffer->offs_size <= sizeof(SymsAddr));
+    syms_buffer_read(buffer, &result, buffer->offs_size);
+    return result;
+}
+
+U32
+syms_hash_djb2(const char *str, u32 size)
+{
+    U32 hash = 5381;
+    const char *str_end = str + size;
+    while (str < str_end) {
+        U32 c = (U32)*str++;
+        hash = ((hash << 5) + hash) + c; /* hash * 33 + c */
+    }
+    return hash;
+}
+
+U16
+syms_hash_djb2_16(const char *str, U32 size)
+{
+    U32 hash = syms_hash_djb2(str, size);
+    U16 result = (U16)(((hash & 0xFFFF0000) >> 16) + (hash & 0x0000FFFF));
+    return result;
+}
+
+#endif /* SYMS_CORE_IMPLEMENT */
diff --git a/Engine/Source/Runtime/SymsLib/include/syms_line_table.h b/Engine/Source/Runtime/SymsLib/include/syms_line_table.h
new file mode 100644
--- /dev/null
+++ b/Engine/Source/Runtime/SymsLib/include/syms_line_table.h
@@ -0,0 +1,35 @@
+// Copyright Epic Games, Inc. All Rights Reserved.
+#ifndef SYMS_LINE_TABLE_INCLUDE_H
+#define SYMS_LINE_TABLE_INCLUDE_H
+
+/******************************************************************************
+ * File   : syms_line_table.h                                                 *
+ * Author : Nikita Smith                                                      *
+ * Created: 2020/08/07                                                        *
+ ******************************************************************************/
+
+#if 0
+typedef struct SymsLineInfo
+{
+  SymsAddr addr;
+  U32 ln;
+  U16 col;
+  U16 file;
+} SymsLineInfo;
+#endif
+
+#if 0
+typedef struct SymsLinesSec 
+{
+  U32 ln_min;
+  U32 ln_max;
+  U32 lines_count;
+  SymsAddr addr_min;
+  SymsAddr addr_max;
+  struct SymsLine *lines;
+  struct SymsLinesSec *next;
+  struct SymsSourceFile file;
+} SymsLinesSec;
+#endif
+
+#endif // SYMS_LINE_TABLE_INCLUDE_H
diff --git a/Engine/Source/Runtime/SymsLib/include/syms_platform.h b/Engine/Source/Runtime/SymsLib/include/syms_platform.h
new file mode 100644
--- /dev/null
+++ b/Engine/Source/Runtime/SymsLib/include/syms_platform.h
@@ -0,0 +1,44 @@
+// Copyright Epic Games, Inc. All Rights Reserved.
+#if !defined(SYMS_PLATFORM_INCLUDE_H)
+#define SYMS_PLATFORM_INCLUDE_H
+
+SYMS_API void
+syms_init_os(void);
+
+SYMS_API SymsUMM
+syms_get_pagesize(void);
+
+SYMS_API void *
+syms_reserve_virtual_memory(SymsUMM size);
+
+SYMS_API syms_bool
+syms_commit_virtual_memory(void *base, SymsUMM size);
+
+SYMS_API void
+syms_free_virtual_memory(void *base, SymsUMM size);
+
+SYMS_INLINE void *
+syms_virtual_alloc(SymsUMM size)
+{
+  void *base = syms_reserve_virtual_memory(size); 
+  if (!syms_commit_virtual_memory(base, size)) {
+    SYMS_ASSERT_FAILURE_PARANOID("unable to commit virtual memory");
+    base = 0;
+  }
+  return base;
+}
+
+typedef struct
+{
+  void *base; // pointer to first byte in the file
+  SymsUMM size; 
+} SymsEntireFile;
+
+SYMS_API SymsEntireFile 
+syms_read_entire_file(const char *file_name);
+
+SYMS_API void
+syms_free_entire_file(SymsEntireFile *file);
+
+#endif // SYMS_PLATFORM_INCLUDE_H
+
