RAD symbol library

diff --git a/Engine/Source/Runtime/SymsLib/include/syms_public.h b/Engine/Source/Runtime/SymsLib/include/syms_public.h
new file mode 100644
--- /dev/null
+++ b/Engine/Source/Runtime/SymsLib/include/syms_public.h
@@ -0,0 +1,1857 @@
+// Copyright Epic Games, Inc. All Rights Reserved.
+#ifndef SYMS_PUBLIC_INCLUDE_H
+#define SYMS_PUBLIC_INCLUDE_H
+
+#ifndef SYMS_API
+#define SYMS_API
+#endif
+
+#ifdef __cplusplus
+extern "C" {
+#endif
+
+#ifdef _MSC_VER
+  typedef unsigned __int64 SYMS_UINT64;
+  typedef unsigned __int32 SYMS_UINT32;
+#else
+  #include <stdint.h>
+  typedef uint64_t SYMS_UINT64;
+  typedef uint32_t SYMS_UINT32;
+#endif
+
+#if defined(__amd64__) || defined(__amd64) || defined(__x86_64__) || defined(_M_X64) || defined(_M_AMD64)
+#define SYMS_X86
+#define SYMS_X64
+#define SYMS_LITTLE_ENDIAN
+#define SYMS_ADDR_MAX 0xffffffffffffffff
+  typedef SYMS_UINT64 SymsAddr;
+#elif defined(i386) || defined(__i386) || defined(__i386__) || defined(_X86_) || defined(_M_I86)
+#define SYMS_X86
+#define SYMS_LITTLE_ENDIAN
+#define SYMS_ADDR_MAX 0xffffff
+  typedef SYMS_UINT32 SymsAddr;
+#elif defined(__arm__) || defined(_M_ARM) 
+#define SYMS_ARM
+#define SYMS_LITTLE_ENDIAN
+  typedef uint32_t SymsAddr;
+#define SYMS_ADDR_MAX 0xffffffff
+#elif defined(__aarch64__)
+#define SYMS_ARM64
+#define SYMS_LITTLE_ENDIAN
+  typedef uint64_t SymsAddr;
+#define SYMS_ADDR_MAX 0xffffffffffffffff
+#else
+#error "arch not defined"
+#endif
+
+typedef SymsAddr SymsUMM;
+
+typedef unsigned int  syms_uint;
+typedef signed   int  syms_int;
+typedef syms_int      syms_bool;
+#define SYMS_INT_MAX  0xffffff
+#define SYMS_UINT_MAX SYMS_INT_MAX
+
+#define SYMS_INVALID_MOD_ID SYMS_INT_MAX
+typedef syms_uint SymsModID; // Module IDentifier
+typedef syms_uint SymsRegID;
+
+typedef struct SymsInstance       SymsInstance;
+typedef struct SymsArena          SymsArena;
+typedef struct SymsArenaFrame     SymsArenaFrame;
+typedef struct SymsBlockAllocator SymsBlockAllocator;
+
+// -----------------------------------------------------------------------------
+
+#define SYMS_VERSION_MAJOR  0
+#define SYMS_VERSION_MINOR  7
+#define SYMS_VERSION_STR    SYMS_TOSTR(SYMS_VERSION_MAJOR) "." SYMS_TOSTR(SYMS_VERSION_MINOR)
+SYMS_API const char * syms_get_version(void);
+
+#define SYMS_RESULT_OK(x)   ((x) == SYMS_ERR_OK)
+#define SYMS_RESULT_FAIL(x) ((x) != SYMS_ERR_OK)
+enum
+{
+  SYMS_ERR_INVALID_CODE_PATH = 0,
+  SYMS_ERR_OK = 1,
+  SYMS_ERR_NOMEM = 2,
+  SYMS_ERR_MAYBE = 3,
+  SYMS_ERR_INVAL = 4,
+  SYMS_ERR_NODATA = 5,
+  SYMS_ERR_NOIMG = 6,
+  SYMS_ERR_NOSYM = 7,
+  SYMS_ERR_NOLINETABLE = 8,
+  SYMS_ERR_DEBUG_ABBREV = 9,
+  SYMS_ERR_UNKNOWN_SYMBOL = 10,
+  SYMS_ERR_INREAD = 11,
+};
+typedef syms_uint SymsErrorCode;
+
+enum 
+{
+  SYMS_ARCH_NULL,
+  SYMS_ARCH_ARM,
+  SYMS_ARCH_ARM32,
+  SYMS_ARCH_PPC,
+  SYMS_ARCH_PPC64,
+  SYMS_ARCH_IA64,
+  SYMS_ARCH_X64,
+  SYMS_ARCH_X86,
+  SYMS_ARCH_COUNT
+};
+typedef syms_uint SymsArch;
+
+typedef struct SymsString
+{
+  syms_uint len;
+  const char *data;
+} SymsString;
+
+//
+// Setup
+//
+
+SYMS_API SymsInstance *
+syms_init_ex(SymsAddr rebase);
+// Creates instance of this library, storage for it is allocated on the heap
+// so it is your resposibilitiy to free memory with syms_quit.
+//
+// - First argument sets base for virtual addresses, image files that have non-zero base address
+//   are remapped. Original base address can be accessed with syms_get_orig_rebase. Dynamic rebase
+//   is supported, you can change base at runtime with syms_set_rebase.
+//
+// Upon success returns instance != NULL.
+
+SYMS_API SymsInstance *
+syms_init(void);
+
+SYMS_API void
+syms_quit(SymsInstance *instance);
+// Frees memory allocated by instance.
+
+// When you have image that was processed by dynamic linker set this flag.
+#define SYMS_LOAD_IMAGE_FLAGS_FROM_MEMORY 0x1
+typedef syms_uint SymsLoadImageFlags;
+
+SYMS_API SymsErrorCode
+syms_load_image(SymsInstance *instance, void *image_base, SymsUMM image_length, SymsLoadImageFlags flags); 
+// Attaches image to instance and if format is windows executable now you can query
+// pdb path with syms_get_pdb_path.
+
+SYMS_API void
+syms_free_image(SymsInstance *instance);
+// Releases allocated memory for image.
+
+// Defers build of lookup tables. 
+//
+// Building lookup tables is the most expensive part of 
+// loading debug info (in terms of CPU cycles). Set this
+// flag if you want to thread this step. A simple job queue
+// will suffice. Retrieve job counter with "syms_get_module_build_count",
+// then, in each worker thread, execute "syms_build_module"
+// with module id that is less then the counter.
+#define SYMS_LOAD_DEBUG_INFO_FLAGS_DEFER_BUILD_MODULE 0x1
+typedef syms_uint SymsLoadDebugInfoFlags;
+
+typedef struct
+{
+  // NOTE(nick): Path or file name of debug info file (e.g. test.pdb, debug_info.dwo).
+  //
+  // It is used to configure inputs for DWARF parser. Debug info itself does not
+  // contain any identifying pieces of data. And so we must resort to using this string
+  // to figure out inputs. In case where image already contains a section that conflicts
+  // with defined file, syms will prioritizes file over section.
+  const char *path;
+
+  // NOTE(nick): Points to the first byte in file.
+  void *base;
+
+  // NOTE(nick): Number of bytes that can be read starting from "base" pointer.
+  SymsUMM size;
+} SymsFile;
+
+// TODO(nick): File API
+#if 0
+SYMS_API SymsErrorCode
+syms_load_debug_info(SymsInstance *instance, char **files, syms_uint file_count, SymsLoadDebugInfoFlags flags);
+#endif
+
+SYMS_API SymsErrorCode
+syms_load_debug_info_ex(SymsInstance *instance, SymsFile *files, syms_uint file_count, SymsLoadDebugInfoFlags flags);
+// Loads debug info from specified files. File data must be kept in memory until you destroy instance or
+// free debug info.
+//
+// "files"     : Array of external debug info files. To get paths to external files use "syms_get_debug_info_paths".
+// "file_count": Number of files.
+// "flags"     : See "SymsLoadDebugInfoFlags"
+//
+// On success returns "syms_true".
+
+SYMS_API void
+syms_free_debug_info(SymsInstance *instance);
+// Releases memory allocated for debug info.
+
+enum 
+{
+  SYMS_IMAGE_NULL, // Image is present
+  SYMS_IMAGE_NT,   // Win32, Win64
+  SYMS_IMAGE_ELF   // Linux
+};
+typedef syms_uint SymsImageType;
+
+// Image loader assings according to magic flags.
+enum 
+{
+  SYMS_IMAGE_HEADER_CLASS_NULL, // Bit-type of image was not resolved detected
+  SYMS_IMAGE_HEADER_CLASS_32,   // 32-bit
+  SYMS_IMAGE_HEADER_CLASS_64    // 64-bit
+};
+typedef syms_uint SymsImageHeaderClass;
+
+typedef char SymsImageImpl[512];
+typedef struct SymsImage 
+{
+  SymsImageType type;
+  SymsImageHeaderClass header_class;
+  SymsArch arch;
+  SymsLoadImageFlags flags;
+  void *img_data;
+  SymsUMM img_data_size;
+  SymsAddr base_addr;
+  SymsImageImpl impl;
+  //char impl[128];
+} SymsImage;
+
+SYMS_API SymsImageType
+syms_get_image_type(const SymsInstance *instance);
+
+SYMS_API SymsImageHeaderClass
+syms_get_image_header_class(const SymsInstance *instance);
+// Get image header class that was read from binary data.
+
+SYMS_API syms_uint
+syms_get_image_size(const SymsInstance *instance);
+// Get image bytes size.
+
+typedef char SymsDebugFileIterImpl[600];
+typedef struct
+{
+  syms_uint file_index;
+  SymsInstance *instance;
+  SymsDebugFileIterImpl impl;
+} SymsDebugFileIter;
+
+SYMS_API syms_bool
+syms_debug_file_iter_init(SymsDebugFileIter *iter, SymsInstance *instance);
+
+SYMS_API syms_bool
+syms_debug_file_iter_next(SymsDebugFileIter *iter, SymsString *path_out);
+
+SYMS_API syms_uint
+syms_get_debug_file_count(SymsInstance *instance);
+
+enum 
+{
+  SYMS_FORMAT_NULL,
+  SYMS_FORMAT_PDB,      // Program Database
+  SYMS_FORMAT_DWARF,    // Dwarf
+  SYMS_FORMAT_ELFSYMTAB // Elf Symbol table
+};
+typedef syms_uint SymsFormatType;
+
+typedef char SymsDebugInfoImpl[1024];
+typedef struct SymsDebugInfo 
+{
+  SymsFormatType type;
+  SymsDebugInfoImpl impl;
+} SymsDebugInfo;
+
+SYMS_API SymsFormatType
+syms_get_symbol_type(const SymsInstance *instance);
+// return value is indentifier for debug info type parser 
+
+SYMS_API SymsImage *
+syms_get_image(SymsInstance *instance);
+
+SYMS_API SymsDebugInfo *
+syms_get_debug_info(SymsInstance *instance);
+
+SYMS_API void
+syms_set_rebase(SymsInstance *instance, SymsAddr rebase);
+// Dynamically rebase addresses in the instance.
+
+SYMS_API SymsAddr
+syms_get_rebase(const SymsInstance *instance);
+// Get current rebase address.
+
+SYMS_API SymsAddr
+syms_get_orig_rebase(const SymsInstance *instance);
+// Original base address that is extracted from image header.
+
+SYMS_API syms_uint
+syms_get_addr_size(SymsInstance *instance);
+// Address size is based on architecture index encoded in image or pdb 
+
+SYMS_API syms_uint
+syms_get_addr_size_ex(SymsArch arch);
+
+SYMS_API SymsArch
+syms_get_arch(SymsInstance *instance);
+
+SYMS_API syms_uint 
+syms_get_proc_count(SymsInstance *instance);
+// Count of infered procedures from debug info.
+
+SYMS_API syms_uint
+  syms_get_line_count(SymsInstance *instance);
+// Counts total number of source lines loaded.
+
+SYMS_API syms_uint
+syms_get_module_build_count(SymsInstance *instance);
+// Return value maximum number of modules that can be built. 
+// if there are no modules to build return value is 0. 
+
+SYMS_API SymsErrorCode
+syms_build_module(SymsInstance *instance, SymsModID mod_id, SymsArena *arena);
+// mod_id: Build lookup tables for module identifier (it is not thread safe to build
+//         two modules with same identifier).
+// 
+// arena: Pointer to allocator.
+//
+// Errors:
+//  SYMS_ERR_INVAL Invalid module identifier, nothing is built.
+// 
+// See also: syms_get_module_build_count
+//
+
+SYMS_API void
+syms_set_rebase(SymsInstance *instance, SymsAddr rebase);
+// Set base for virtual addresses.
+
+SYMS_API SymsAddr
+syms_get_rebase(const SymsInstance *instance);
+
+SYMS_API SymsAddr
+syms_get_orig_rebase(const SymsInstance *instance);
+// Returns address at which compiler/linker based addresses for symbols.
+//
+// If image is compiled to be relocatable at runtime then original base is zero.
+
+SYMS_API syms_bool
+syms_match_image_with_debug_info(const SymsImage *image, const SymsDebugInfo *debug_info);
+
+// --------------------------------------------------------------------------------
+
+enum 
+{
+  SYMS_STRING_REF_TYPE_NULL,
+
+  // NOTE(nick): Regular string
+  SYMS_STRING_REF_TYPE_STR,
+
+  // NOTE(nick): Reference into pdb file, for more info look into libs/pdb.h
+  SYMS_STRING_REF_TYPE_PDB,
+
+  // NOTE(nick): This is a hack for dwarf strings encoded in line table.
+  // Directory and file name are stored in different locations, and requires
+  // decoder, us, to merge them at runtime.
+  SYMS_STRING_REF_TYPE_DW_PATH
+};
+typedef syms_uint SymsStringRefType;
+
+typedef char SymsStringRefImpl[64];
+typedef struct SymsStringRef 
+{
+  SymsStringRefType type;
+  SymsStringRefImpl impl;
+} SymsStringRef;
+
+SYMS_API syms_uint
+syms_strref_get_size(SymsInstance *instance, SymsStringRef *ref);
+// Retrieves number of bytes string references takes up.
+
+SYMS_API syms_uint
+syms_read_strref(SymsInstance *instance, SymsStringRef *syms_ref, void *bf, syms_uint bf_max);
+// Reads string reference to a buffer. If size of the buffer is smaller then
+// size of the string reference this function will read as many bytes as possible.
+//
+// On success returns number of written bytes including '\0'.
+
+SYMS_API char *
+syms_read_strref_arena(SymsInstance *instance, SymsStringRef *ref, SymsArena *arena);
+
+SYMS_API SymsString
+syms_read_strref_str(SymsInstance *instance, SymsStringRef *ref, SymsArena *arena);
+
+SYMS_API syms_bool
+syms_string_cmp_strref(SymsInstance *instance, SymsString a, SymsStringRef b);
+
+SYMS_API SymsStringRef
+syms_string_ref_str(SymsString str);
+
+// --------------------------------------------------------------------------------
+
+enum
+{
+  SYMS_TYPE_ID_NULL,
+  SYMS_TYPE_ID_PDB,
+  SYMS_TYPE_ID_DW,
+  SYMS_TYPE_ID_ELF,
+  SYMS_TYPE_ID_BUILTIN,
+  SYMS_TYPE_ID_COUNTER,
+};
+typedef syms_uint SymsTypeIDKind;
+
+typedef char SymsTypeIDImpl[32];
+typedef struct
+{
+  syms_uint kind;
+  SymsTypeIDImpl impl;
+} SymsTypeID;
+
+typedef SymsTypeID SymsTypeRef;
+
+// --------------------------------------------------------------------------------
+// Type Map
+
+typedef enum 
+{
+  SYMS_TYPE_NULL,
+  /* Integers */
+  SYMS_TYPE_INT8,
+  SYMS_TYPE_INT16,
+  SYMS_TYPE_INT32,
+  SYMS_TYPE_INT64,
+  SYMS_TYPE_INT128,
+  SYMS_TYPE_INT256,
+  SYMS_TYPE_INT512,
+  SYMS_TYPE_UINT8,
+  SYMS_TYPE_UINT16,
+  SYMS_TYPE_UINT32,
+  SYMS_TYPE_UINT64,
+  SYMS_TYPE_UINT128,
+  SYMS_TYPE_UINT256,
+  SYMS_TYPE_UINT512,
+  /* Floats */
+  SYMS_TYPE_FLOAT16,
+  SYMS_TYPE_FLOAT32,
+  SYMS_TYPE_FLOAT32PP, /* NOTE(nick): Partial precision */
+  SYMS_TYPE_FLOAT48,
+  SYMS_TYPE_FLOAT64,
+  SYMS_TYPE_FLOAT80,
+  SYMS_TYPE_FLOAT128,
+  SYMS_TYPE_CHAR,
+  SYMS_TYPE_VOID,
+  SYMS_TYPE_BOOL,
+  SYMS_TYPE_PTR,
+  SYMS_TYPE_ARR,
+  SYMS_TYPE_ENUM,
+  SYMS_TYPE_PROC,
+  SYMS_TYPE_PROC_PARAM,
+  SYMS_TYPE_TYPEDEF,
+  SYMS_TYPE_STRUCT,
+  SYMS_TYPE_UNION,
+  SYMS_TYPE_CLASS,
+  SYMS_TYPE_METHOD,
+  SYMS_TYPE_VIRTUAL_TABLE,
+  SYMS_TYPE_BASE_CLASS,
+  SYMS_TYPE_BITFIELD,
+  SYMS_TYPE_COMPLEX32,
+  SYMS_TYPE_COMPLEX64,
+  SYMS_TYPE_COMPLEX80,
+  SYMS_TYPE_COMPLEX128,
+  SYMS_TYPE_VARIADIC,
+
+  SYMS_TYPE_STRING,
+  SYMS_TYPE_WCHAR,
+
+  SYMS_TYPE_ARRPTR,
+
+  SYMS_TYPE_INVALID
+} SymsTypeKind_e;
+typedef syms_uint SymsTypeKind;
+
+enum
+{
+  SYMS_TYPE_MDFR_NULL = 0,
+
+  // NOTE(nick): atomic qualified type.
+  SYMS_TYPE_MDFR_ATOMIC = (1 << 0),
+
+  // NOTE(nick): const qualified type.
+  SYMS_TYPE_MDFR_CONST = (1 << 1),
+
+  // NOTE(nick): Same as const, but the data lives in read-only part of the RAM. 
+  SYMS_TYPE_MDFR_IMMUTABLE = (1 << 2),
+
+  // NOTE(nick): When a data type is marked for compiler to use as little space as possible.
+  SYMS_TYPE_MDFR_PACKED = (1 << 3),
+
+  // NOTE(nick): A reference to (lvalue of) an object of the type being modifid. 
+  SYMS_TYPE_MDFR_REF = (1 << 4),
+
+  // NOTE(nick): restrict qualified type.
+  SYMS_TYPE_MDFR_RESTRICT = (1 << 5),
+
+  // NOTE(nick): rvalue ref to an object of the type being modified (for example, in C++)
+  SYMS_TYPE_MDFR_RVALUE_REF = (1 << 6),
+
+  // NOTE(nick): shared qualified type. 
+  SYMS_TYPE_MDFR_SHARED = (1 << 7),
+
+  // NOTE(nick): volatile qualified type. 
+  SYMS_TYPE_MDFR_VOLATILE = (1 << 8),
+
+  // NOTE(nick): Assigned when memory values of this type can be displayed as text.
+  SYMS_TYPE_MDFR_CHAR = (1 << 9),
+
+  // NOTE(nick): set when nested UDT is encountered
+  SYMS_TYPE_MDFR_NESTED = (1 << 10),
+
+  SYMS_TYPE_MDFR_FWDREF = (1 << 11),
+
+  SYMS_TYPE_MDFR_INVALID
+};
+typedef syms_uint SymsTypeModifier;
+
+typedef struct
+{
+  syms_uint arg_count;            // Argumnet count
+  SymsTypeID *arglist_id;  // Argument list
+  SymsTypeID *return_id;   // Return type
+} SymsTypeProc;
+
+typedef struct
+{
+  syms_uint arg_count;            // Argument count
+  SymsTypeID *arglist_id;  // Argument list
+  SymsTypeID *return_id;   // Return type
+  SymsTypeID *class_id;    // OOP Class
+  SymsTypeID *this_id;     // OOP this
+} SymsTypeMethod;
+
+
+// Array counter, when array has more dimensions entry_id 
+// points to next dimension
+typedef struct
+{
+  SymsUMM entry_count;         
+  SymsTypeID *entry_id;
+} SymsTypeArray;
+
+enum
+{
+  SYMS_MEMBER_TYPE_NULL,
+
+  /* Member describes memory region inside UDT */
+  SYMS_MEMBER_TYPE_DATA,
+
+  /* Static data member of a UDT */
+  SYMS_MEMBER_TYPE_STATIC_DATA,
+
+  SYMS_MEMBER_TYPE_ENUM,
+
+  /* Method info of a UDT  */
+  SYMS_MEMBER_TYPE_METHOD,
+
+  /* Type that describes layout of the virtual table.
+   * Compiler might choose to store virtual table
+   * at an address other than zero, so add
+   * offset to the UDT base address when formatting
+   * memory for it. */
+  SYMS_MEMBER_TYPE_VTABLE,
+
+  /* Base class location, use type id to format memory for it. */
+  SYMS_MEMBER_TYPE_BASE_CLASS,
+
+  /* Type that was declared inside the UDT */
+  SYMS_MEMBER_TYPE_NESTED_TYPE
+};
+typedef syms_uint SymsMemberType;
+
+enum
+{
+  SYMS_MEMBER_ACCESS_NULL,
+  SYMS_MEMBER_ACCESS_PRIVATE,
+  SYMS_MEMBER_ACCESS_PUBLIC,
+  SYMS_MEMBER_ACCESS_PROTECTED
+};
+typedef syms_uint SymsMemberAccess;
+
+enum
+{
+  SYMS_MEMBER_MODIFIER_NULL,
+
+  /* No modifiers, just a regular method */
+  SYMS_MEMBER_MODIFIER_VANILLA,
+
+  /* Method marked with static */
+  SYMS_MEMBER_MODIFIER_STATIC,
+
+  /* Friend method */
+  SYMS_MEMBER_MODIFIER_FRIEND,
+
+  /* Virtual method */
+  SYMS_MEMBER_MODIFIER_VIRTUAL,
+
+  /* Virtual method provides only signature: void virtual foo(void) = 0; */
+  SYMS_MEMBER_MODIFIER_PURE_INTRO,
+
+  /* TODO: ??? */
+  SYMS_MEMBER_MODIFIER_PURE_VIRTUAL,
+  SYMS_MEMBER_MODIFIER_INTRO
+};
+typedef syms_uint SymsMemberModifier;
+
+typedef enum
+{
+  SYMS_VTABLE_ENTRY_NULL,
+
+  SYMS_VTABLE_ENTRY_PTR16,
+  SYMS_VTABLE_ENTRY_PTR32,
+
+  /* Pointer is composed out of two parts 
+   * first is segment and second is offset.
+   * Add them together to get final pointer. */
+  SYMS_VTABLE_ENTRY_SEGOFF16, 
+  SYMS_VTABLE_ENTRY_SEGOFF32,
+
+  /* CodeView specific types, more in syms_codeview.h */
+  SYMS_VTABLE_ENTRY_THIN,
+  SYMS_VTABLE_ENTRY_OUTER,
+  SYMS_VTABLE_ENTRY_META
+} SymsVTableEntryType_e;
+typedef syms_uint SymsVTableEntryType;
+
+typedef struct
+{
+  /* Label given for the UDT member */
+  char *label;
+
+  /* Member type index */
+  SymsTypeID *type_id;
+
+  /* Type index for the field in the descriminated union */
+  SymsMemberType type;
+
+  /* Exports member access parameters when parsing class,
+   * in other cases it is set to SYMS_MEMBER_ACCESS_NULL */
+  SymsMemberAccess access;
+
+  /* Exports member modifiers when parsing class,
+   * in other cases it is set to null */
+  SymsMemberModifier modifier;
+
+  union
+  {
+    /* Virtual table format description. */
+    struct
+    {
+      SymsAddr offset;            /* virtual table offset */
+      syms_uint count;            /* number of entries */
+      SymsVTableEntryType *base;  /* first entry */
+    } vtab;
+
+    struct
+    {
+      SymsUMM offset;
+    } data;
+
+    struct
+    {
+      SymsUMM value;
+    } enum_value;
+
+    struct
+    {
+      SymsUMM offset;
+    } base_class;
+
+    struct
+    {
+      SymsAddr vbaseoff;
+    } method;
+  } u;
+} SymsTypeMember;
+
+struct SymsTypeMemberNode
+{
+  SymsTypeMember data;
+  struct SymsTypeMemberNode *next;
+};
+
+typedef struct
+{
+  syms_uint count;
+  SymsTypeMember *base;
+} SymsTypeUDT;
+
+typedef struct
+{
+  SymsTypeID *base_type_id;
+  syms_uint len;
+  syms_uint pos;
+} SymsTypeBitfield;
+
+typedef struct
+{
+  char *label;
+  SymsTypeID *type_id;
+  SymsTypeModifier modifier;
+  SymsTypeKind kind;
+  syms_uint size;
+  union
+  {
+    SymsTypeProc   proc;
+    SymsTypeMethod method;
+    SymsTypeArray  array;
+    SymsTypeUDT    udt; // SYMS_TYPE_STRUCT, SYMS_TYPE_CLASS, SYMS_TYEPE_UNION
+    SymsTypeBitfield bitfield;
+  } u;
+} SymsSymbolTypeInfo;
+
+typedef struct
+{
+  char *name;      // section label
+  SymsAddr off;    // In-file offset where section starts
+  SymsAddr va;     // In-memory address where section starts
+  SymsAddr length; // Number of bytes that make up section
+} SymsSymbolSection;
+
+typedef struct
+{
+  // [lo, hi]
+  SymsAddr lo;
+  SymsAddr hi;
+} SymsSymbolLexicalBlock;
+
+enum
+{
+  SYMS_PTR_MODE_NULL,
+  SYMS_PTR_MODE_NORMAL, /* Regular pointer*/
+  SYMS_PTR_MODE_LVREF,  /* L-value reference */
+  SYMS_PTR_MODE_RVREF,  /* R-value reference */
+  SYMS_PTR_MODE_MEM,    /* Pointer to a data member */
+  SYMS_PTR_MODE_MFUNC   /* Pointer to member function */
+
+};
+typedef syms_uint SymsPointerMode;
+
+typedef char SymsTypeImpl[128];
+typedef struct SymsType 
+{
+  // NOTE(nick): A unique identifier for this type.
+  SymsTypeID id;
+
+  // NOTE(nick): A unique identifier for next type in the hierarchy chain.
+  SymsTypeID next_id;
+
+  SymsTypeKind kind;
+
+  SymsTypeModifier modifier;
+
+  // NOTE(nick): Size of the type itself.
+  // If "kind" is SYMS_TYPE_ARR then this field will contain count of the elements.
+  syms_uint size;
+
+  syms_bool is_fwdref;
+
+  // NOTE(nick): Name of a type.
+  SymsStringRef name_ref;
+
+  SymsStringRef decl_file;
+  syms_uint decl_ln;
+
+  union 
+  {
+    struct 
+    {
+      syms_uint arg_count;
+      SymsTypeID arglist_type_id;
+      SymsTypeID ret_type_id;
+    } proc; // NOTE(nick): SYMS_TYPE_PROC
+
+    struct
+    {
+      SymsPointerMode mode;
+    } ptr;
+
+    struct 
+    {
+      syms_uint arg_count;
+      SymsTypeID ret_type_id;
+      SymsTypeID class_type_id;
+      SymsTypeID this_type_id;
+      SymsTypeID arglist_type_id;
+    } method;
+
+    struct
+    {
+      SymsTypeID base_type_id;
+      syms_uint len;
+      syms_uint pos;
+    } bitfield;
+
+    SymsTypeID param_type;  // NOTE(nick): SYMS_TYPE_PROC_PARAM
+  } u;
+
+  SymsTypeImpl impl;
+} SymsType;
+
+typedef char SymsTypeIterImpl[64];
+typedef struct SymsTypeIter 
+{
+  SymsDebugInfo *debug_info;
+  SymsTypeIterImpl impl;
+} SymsTypeIter;
+
+SYMS_API syms_bool
+syms_type_iter_init(SymsInstance *instance, SymsTypeIter *iter);
+
+SYMS_API syms_bool
+syms_type_iter_next(SymsTypeIter *iter, SymsTypeID *id_out);
+
+SYMS_API syms_int
+syms_typeid_cmp(const SymsTypeID *a, const SymsTypeID *b);
+
+SYMS_API syms_bool
+syms_infer_type(SymsInstance *instance, SymsTypeID type_id, SymsType *type_out);
+
+SYMS_API void
+syms_infer_builtin_type(SymsArch arch, SymsTypeKind kind, SymsType *type_out);
+
+SYMS_API const char *
+syms_get_type_kind_info(SymsTypeKind kind);
+
+SYMS_API SymsType
+syms_infer_base_type(SymsInstance *instance, const SymsType *type);
+
+SYMS_API syms_uint
+syms_get_type_size(SymsInstance *instance, SymsType *type);
+
+SYMS_API SymsTypeID
+syms_type_id_for_kind(SymsTypeKind kind);
+
+SYMS_API SymsType
+syms_make_ptr_type(SymsArch arch, const SymsTypeID *ref);
+
+SYMS_API syms_bool
+syms_type_from_name(SymsInstance *instance, const char *name, syms_uint name_size, SymsType *type_out);
+
+SYMS_API SymsTypeID
+syms_new_type(SymsInstance *instance);
+
+
+//
+// Symbols
+//
+
+enum
+{
+  SYMS_SYMBOL_NULL,
+  SYMS_SYMBOL_SRCMAP,
+  SYMS_SYMBOL_PROC,
+  SYMS_SYMBOL_VAR_LOCAL,
+  SYMS_SYMBOL_VAR,
+  SYMS_SYMBOL_MODULE,
+  SYMS_SYMBOL_SECTION,
+  SYMS_SYMBOL_CONST_DATA,
+  SYMS_SYMBOL_STATIC_DATA,
+  SYMS_SYMBOL_TYPE_INFO,
+  SYMS_SYMBOL_LEXICAL_BLOCK,
+  SYMS_SYMBOL_TRAMPOLINE,
+  // Special case for PDB that might have public symbol table,
+  // but other symbols stripped. It exports range with label that
+  // is mangled by MSVC. 
+  SYMS_SYMBOL_PUBLIC
+};
+typedef syms_uint SymsSymbolKind;
+
+enum 
+{
+  SYMS_CHECKSUM_NULL,
+  SYMS_CHECKSUM_MD5,
+  SYMS_CHECKSUM_SHA1,
+  SYMS_CHECKSUM_SHA256
+};
+typedef syms_uint SymsChecksumType;
+
+typedef struct 
+{
+  SymsAddr addr;
+  SymsUMM ln;
+  SymsUMM col;
+  SymsUMM instructions_size;
+  SymsUMM chksum_size;
+  void *chksum;
+  char *path;
+  SymsChecksumType chksum_type; 
+} SymsSymbolSrcmap;
+
+enum 
+{
+  SYMS_VAR_FLAG_PARAM        = (1 << 0),
+  SYMS_VAR_FLAG_COMPILER_GEN = (1 << 1),
+  SYMS_VAR_FLAG_ALIASED      = (1 << 2),
+  SYMS_VAR_FLAG_RETVAL       = (1 << 3),
+  SYMS_VAR_FLAG_OPT_OUT      = (1 << 4),
+  SYMS_VAR_FLAG_STATIC       = (1 << 5),
+  SYMS_VAR_FLAG_GLOBAL       = (1 << 6)
+};
+typedef syms_uint SymsVarFlags;
+
+typedef char SymsEncodedLocationImpl[64];
+typedef struct
+{
+  SymsEncodedLocationImpl impl;
+} SymsEncodedLocation;
+
+typedef struct
+{
+  SymsAddr lo;
+  SymsAddr hi;
+} SymsRangeGap;
+
+typedef struct SymsRangeData
+{
+  syms_uint gap_count;
+  SymsRangeGap *gap_base;
+  SymsEncodedLocation *encoded_va;
+  SymsAddr lo;
+  SymsAddr hi;
+} SymsRangeData;
+
+typedef struct SymsRangeNode
+{
+  SymsRangeData data;
+  struct SymsRangeNode *next;
+} SymsRangeNode;
+
+typedef char SymsRangeToken[32];
+
+typedef struct SymsSymbolVar
+{
+  SymsVarFlags flags;
+
+  // user gieven/compiler generated label for variable
+  char *label;
+
+  // type might be missing and in that case it is set to null
+  SymsTypeID *type_id;
+
+  //SymsRangeIter range_iter;
+
+  SymsRangeNode *range_list;
+
+  // TODO: get rid of these:
+  SymsEncodedLocation *encoded_va;
+  struct SymsSymbolVar *next;
+} SymsSymbolVar;
+
+// Symbol for procedure exported from low-level SymsProc.
+// Range is guaranteed to be present at all times and everything else is optional.
+typedef struct SymsSymbolProc
+{
+
+  // Address of first instruction.
+  SymsAddr range_lo;
+
+  // Address of last instruction.
+  SymsAddr range_hi;
+
+  // Address past procedure prolog. A hint for debuggers
+  // to avoid stepping through stack setup of procedures.
+  SymsAddr dbg_start;
+
+  // address before procedure epilog.
+  SymsAddr dbg_end;
+
+  // Label attached to symbol by user or compiler.
+  char *label;
+
+  syms_uint local_count;
+
+  // Resolved inline chain.
+  struct SymsSymbolProc *inline_chain;
+
+  // Local variables with respect to lexical-blocks.
+  SymsSymbolVar *locals;
+
+  // Source file location
+  SymsSymbolSrcmap srcmap;
+
+  // Type id for symbol.
+  SymsTypeID *type_id;
+} SymsSymbolProc;
+
+typedef struct
+{
+  SymsAddr jump_addr;
+} SymsSymbolTrampoline;
+
+typedef struct 
+{
+  char *name;
+} SymsSymbolModule;
+
+typedef struct
+{
+  char *label;
+  SymsTypeID *type_id;
+  SymsAddr addr;
+} SymsSymbolStaticData;
+
+typedef struct
+{
+  char *label;
+  void *value;
+  SymsTypeID *type_id;
+  SymsUMM value_size;
+} SymsSymbolConstData;
+
+typedef struct SymsSymbol
+{
+  // Module ID where symbol is stored. When it is set to SYMS_INVALID_MOD_ID
+  // symbol does not have a module.
+  SymsModID mod_id;
+
+  // Indicates which union field is valid.
+  SymsSymbolKind kind;
+
+  struct SymsSymbol *next;
+
+  union
+  {
+    SymsSymbolLexicalBlock lexical_block;
+    SymsSymbolVar         var_local;
+    SymsSymbolProc        proc;
+    SymsSymbolProc        public_proc;
+    SymsSymbolModule      mod;
+    SymsSymbolSection     sec;
+    SymsSymbolSrcmap      source_map;
+    SymsSymbolConstData   cdata;
+    SymsSymbolStaticData  sdata;
+    SymsSymbolTypeInfo    type_info;
+    SymsSymbolTrampoline  trampoline;
+    SymsSymbolSection     section;
+  } u;
+} SymsSymbol;
+
+typedef struct SymsTypeIDNode
+{
+  SymsTypeID *data;
+  struct SymsTypeIDNode *next; 
+} SymsTypeIDNode;
+
+typedef struct
+{
+  syms_uint count;
+  SymsTypeIDNode *first;
+} SymsTypeList;
+
+// ----------------------------------------------------------------------------
+
+enum
+{
+  SYMS_RANGE_NULL,
+  SYMS_RANGE_PLAIN,
+  SYMS_RANGE_IMPL
+};
+typedef syms_uint SymsRangeType;
+
+typedef char SymsRangeImpl[64];
+typedef struct SymsRange 
+{
+  SymsRangeType type;
+  union 
+  {
+    struct 
+    {
+      SymsAddr lo;
+      SymsAddr hi;
+    } plain;
+    SymsRangeImpl impl;
+  } u;
+} SymsRange;
+
+typedef char SymsRangeIterImpl[64];
+typedef struct SymsRangeIter
+{
+  SymsInstance *instance;
+  SymsDebugInfo *debug_info;
+  SymsRange range;
+  SymsRangeIterImpl impl;
+} SymsRangeIter;
+
+SYMS_API syms_bool
+syms_range_iter_init(SymsRangeIter *iter, SymsInstance *instance, SymsRange *range);
+
+SYMS_API syms_bool
+syms_range_iter_next(SymsRangeIter *iter, SymsAddr *lo_out, SymsAddr *hi_out);
+
+// ----------------------------------------------------------------------------
+
+enum 
+{
+  SYMS_LOCATION_NULL,
+  SYMS_LOCATION_IMPLICIT,
+  SYMS_LOCATION_INDIRECT,
+  SYMS_LOCATION_VA
+};
+typedef syms_uint SymsLocationKind;
+
+#define SYMS_LOCATION_IMPLICIT_VALUE_MAX 32
+typedef struct SymsLocation 
+{
+  SymsLocationKind kind;
+  union 
+  {
+    struct 
+    {
+      syms_uint len;
+      char data[SYMS_LOCATION_IMPLICIT_VALUE_MAX];
+    } implicit;
+
+    struct 
+    {
+      syms_uint len;
+      void const *data;
+    } indirect;
+
+    SymsAddr va;
+  } u;
+} SymsLocation;
+
+#define SYMS_MEMREAD_SIG(name) SymsErrorCode name(void *user_context, SymsAddr va, void *buffer, syms_uint buffer_size)
+typedef SYMS_MEMREAD_SIG(syms_memread_sig);
+
+#define SYMS_REGREAD_SIG(name) syms_uint name(void *user_context, SymsArch arch, SymsRegID regid, void *buffer, syms_uint buffer_size)
+typedef SYMS_REGREAD_SIG(syms_regread_sig);
+
+#define SYMS_REGWRITE_SIG(name) syms_uint name(void *user_context, SymsArch arch, SymsRegID regid, void *buffer, syms_uint buffer_size)
+typedef SYMS_REGWRITE_SIG(syms_regwrite_sig);
+
+SYMS_API SymsErrorCode
+syms_decode_location(SymsInstance *instance, SymsEncodedLocation *loc, void *regread_context, syms_regread_sig *regread_callback, void *memread_context, syms_memread_sig *memread_callback, SymsLocation *loc_out);
+
+// ----------------------------------------------------------------------------
+
+typedef char SymsModImpl[512];
+typedef struct SymsMod 
+{
+  SymsModID id;
+  SymsStringRef name;
+  SymsAddr va;
+  SymsUMM size;
+  SymsModImpl impl;
+} SymsMod;
+
+typedef char SymsModIterImpl[128];
+typedef struct SymsModIter 
+{
+  SymsInstance *instance;
+  syms_uint index;
+  SymsModIterImpl impl;
+} SymsModIter;
+
+SYMS_API syms_bool
+syms_mod_iter_init(SymsInstance *instance, SymsModIter *iter);
+
+SYMS_API syms_bool
+syms_mod_iter_next(SymsModIter *iter, SymsMod *mod_out);
+
+// ----------------------------------------------------------------------------
+
+typedef struct SymsLine 
+{
+  // NOTE(nick): Address of the first instruction byte
+  SymsAddr va;
+
+  // NOTE(nick): Line number that contributed instructions
+  syms_uint ln;
+
+  // NOTE(nick): Column number
+  syms_uint col;
+} SymsLine;
+
+typedef struct SymsSourceFile 
+{
+  SymsStringRef name;
+  SymsChecksumType chksum_type;
+  char chksum[32];
+} SymsSourceFile;
+
+typedef struct SymsSourceFileMap 
+{
+  SymsSourceFile file;
+  SymsLine line;
+  syms_uint instructions_size;
+} SymsSourceFileMap;
+
+typedef char SymsLineIterImpl[512];
+typedef struct SymsLineIter 
+{
+  syms_int switched_file;
+  syms_int has_line_count;
+  syms_uint line_count;
+  SymsInstance  *instance;
+  SymsDebugInfo *debug_info;
+  SymsArena     *arena;
+  SymsSourceFile file;
+  SymsLineIterImpl impl;
+} SymsLineIter;
+
+SYMS_API syms_bool
+syms_line_iter_init(SymsInstance *instance, SymsLineIter *iter, SymsMod *mod);
+// Initializes line iterator for given module.
+//
+// On success returns "syms_true".
+
+SYMS_API syms_bool
+syms_line_iter_next(SymsLineIter *iter, SymsLine *line_out);
+// Fetches next line from the line table.
+//
+// On success returns "syms_true".
+
+// ----------------------------------------------------------------------------
+
+typedef struct SymsMethod
+{
+  SymsTypeID type_id;
+  SymsStringRef name_ref;
+  syms_uint vbaseoff;
+} SymsMethod;
+
+typedef struct SymsVirtualTable
+{
+  SymsAddr offset;
+} SymsVirtualTable;
+
+#define SYMS_MEMBER_ENUM_MAX 32
+typedef struct SymsMember 
+{
+  SymsMemberAccess access;
+  SymsMemberModifier modifier;
+  SymsMemberType type;
+
+  // NOTE(nick): Type of the member.
+  SymsTypeID type_id;
+
+  // NOTE(nick): Member name that was given by the user.
+  SymsStringRef name_ref;
+
+  union 
+  {
+    // NOTE(nick): When member owner is strcut/union/class this is offset of member.
+    SymsUMM data_offset;
+
+    SymsMethod method;
+
+    SymsVirtualTable vtab;
+
+    struct
+    {
+      SymsUMM offset;
+    } base_class;
+
+    // NOTE(nick): Format value according to type_id.
+    char enum_value[SYMS_MEMBER_ENUM_MAX];
+  } u;
+} SymsMember;
+
+typedef char SymsMemberIterImpl[256];
+typedef struct SymsMemberIter 
+{
+  SymsInstance *instance;
+  SymsDebugInfo *debug_info;
+  SymsMemberIterImpl impl;
+} SymsMemberIter;
+
+SYMS_API syms_bool
+syms_member_iter_init(SymsInstance *instance, SymsMemberIter *iter, SymsType *type);
+// Parser for user defined types (struct, class, union, enum)
+
+SYMS_API syms_bool
+syms_member_iter_next(SymsMemberIter *iter, SymsMember *syms_member);
+
+// ----------------------------------------------------------------------------
+
+// TODO(nick): Merge this var with SymsVar
+typedef struct SymsGlobal 
+{
+  SymsTypeID type_id;
+  SymsStringRef name;
+  SymsEncodedLocation encoded_va;
+} SymsGlobal;
+typedef struct SymsGlobal SymsLocalData;
+
+typedef char SymsGlobalIterImpl[128];
+typedef struct SymsGlobalIter 
+{
+  SymsDebugInfo *debug_info;
+  SymsGlobalIterImpl impl;
+} SymsGlobalIter;
+
+SYMS_API syms_bool
+syms_global_iter_init(SymsInstance *instance, SymsGlobalIter *iter);
+
+SYMS_API syms_bool
+syms_global_iter_next(SymsGlobalIter *iter, SymsGlobal *global_out);
+
+typedef char SymsLocalDataIterImpl[64];
+typedef struct SymsLocalDataIter 
+{
+  SymsModID mod_id;
+  SymsDebugInfo *debug_info;
+  SymsLocalDataIterImpl impl;
+} SymsLocalDataIter;
+
+SYMS_API syms_bool
+syms_local_data_iter_init(SymsInstance *instance, SymsLocalDataIter *iter, SymsMod *mod);
+
+SYMS_API syms_bool
+  syms_local_data_iter_next(SymsLocalDataIter *iter, SymsLocalData *ldata_out);
+
+SYMS_API syms_bool
+syms_global_from_name(SymsInstance *instance, const char *name, syms_uint name_size, SymsGlobal *gvar_out);
+
+// ----------------------------------------------------------------------------
+
+typedef char SymsProcImpl[128];
+typedef struct SymsProc 
+{
+  // NOTE(nick): Type ID of a procedure.
+  SymsTypeID type_id;
+
+  // NOTE(nick): Address of a first instruction.
+  SymsAddr va;
+
+  // NOTE(nick): Byte length of all instructions that a procedure encapsulates.
+  syms_uint len;
+
+  SymsAddr dbg_start_va;
+  SymsAddr dbg_end_va;
+
+  // NOTE(nick): Name ref of a procedure.
+  SymsStringRef name_ref;
+
+  SymsRange range;
+
+  // NOTE(nick): For internal use only.
+  SymsProcImpl impl;
+} SymsProc;
+
+typedef char SymsProcIterImpl[256];
+typedef struct SymsProcIter 
+{
+  SymsInstance *instance;
+  SymsProcIterImpl impl;
+} SymsProcIter;
+
+SYMS_API syms_bool
+syms_proc_iter_init(SymsInstance *instance, SymsProcIter *iter, SymsMod *mod);
+
+SYMS_API syms_bool
+syms_proc_iter_next(SymsProcIter *iter, SymsProc *proc);
+
+// ----------------------------------------------------------------------------
+
+typedef char SymsScopeImpl[64];
+typedef struct SymsScope
+{
+  SymsAddr inst_lo;
+  SymsAddr inst_hi;
+  SymsScopeImpl impl;
+} SymsScope;
+
+typedef struct SymsVar 
+{
+  // NOTE(nick): Type of a local variable.
+  SymsTypeID type_id;
+
+  SymsVarFlags flags;
+
+  // NOTE(nick): Use "syms_decode_location" routine to decode it.
+  SymsEncodedLocation encoded_va;
+
+  SymsRange range;
+
+  // NOTE(nick): Reference to name of a variable.
+  SymsStringRef name_ref;
+} SymsVar;
+
+enum 
+{
+  SYMS_LOCAL_EXPORT_NULL,
+  SYMS_LOCAL_EXPORT_SCOPE,
+  SYMS_LOCAL_EXPORT_VAR,
+  SYMS_LOCAL_EXPORT_SCOPE_END,
+};
+typedef syms_uint SymsLocalExportType;
+
+typedef struct SymsLocalExport
+{
+  SymsLocalExportType type;
+  union
+  {
+    SymsScope scope;
+    SymsVar   var;
+  } u;
+} SymsLocalExport;
+
+typedef char SymsLocalIterImpl[512];
+typedef struct SymsLocalIter 
+{
+  SymsInstance *instance;
+  SymsLocalIterImpl impl;
+} SymsLocalIter;
+
+SYMS_API syms_bool
+syms_local_iter_init(SymsInstance *instance, SymsLocalIter *iter, SymsProc *proc, void *stack, syms_uint stack_size);
+
+
+SYMS_API syms_bool
+syms_local_iter_next(SymsLocalIter *local_iter, SymsLocalExport *export_out);
+// Fetches next local-variable symbol.
+//
+// On successful decoding of a local-variable symbol it returns "syms_true".
+
+// ----------------------------------------------------------------------------
+
+typedef char SymsFileIterImpl[700];
+typedef struct SymsFileIter 
+{
+  SymsDebugInfo *debug_info;
+  SymsFileIterImpl impl;
+} SymsFileIter;
+
+SYMS_API syms_bool
+syms_file_iter_init(SymsInstance *instance, SymsFileIter *iter);
+
+SYMS_API syms_bool
+syms_file_iter_next(SymsFileIter *iter, SymsStringRef *ref_out);
+
+// ----------------------------------------------------------------------------
+
+typedef char SymsArgIterImpl[128];
+typedef struct SymsArgIter 
+{
+  SymsDebugInfo *debug_info;
+  SymsArgIterImpl impl;
+} SymsArgIter;
+
+SYMS_API syms_bool
+syms_arg_iter_init(SymsInstance *instance, SymsArgIter *iter, SymsType *proc_type);
+
+SYMS_API syms_bool
+syms_arg_iter_next(SymsArgIter *iter, SymsTypeID *type_id_out);
+
+
+// ----------------------------------------------------------------------------
+
+typedef struct SymsInlineSite 
+{
+  // inline site name
+  SymsStringRef name;
+
+  // type index of inline site
+  SymsTypeID type_id;
+
+  // address of first instruction in inlined block
+  SymsAddr range_lo;
+  // address where inlined block ends (includes last instruction)
+  SymsAddr range_hi;
+
+  // source file name where procedure is inlined
+  SymsSourceFile call_file;
+  // source file name where procedure is declare
+  SymsSourceFile decl_file;
+  // line number in source file where compiler inlined procedure
+  syms_uint call_ln;
+  // line number in source file inline site is declared
+  syms_uint decl_ln;
+
+  SymsAddr sort_index;
+
+  SymsProc inlinee;
+  SymsSourceFileMap src;
+} SymsInlineSite;
+
+typedef char SymsInlineIterImpl[512]; 
+typedef struct SymsInlineIter 
+{
+  SymsInstance *instance;
+  SymsDebugInfo *debug_info;
+  SymsInlineIterImpl impl;
+} SymsInlineIter;
+
+SYMS_API syms_bool
+syms_inline_iter_init(SymsInstance *instance, SymsInlineIter *inline_iter, SymsAddr pc);
+
+SYMS_API syms_bool
+syms_inline_iter_next(SymsInlineIter *inline_iter, SymsInlineSite *site_out);
+
+// ----------------------------------------------------------------------------
+
+#define SYMS_CONST_VALUE_MAX 32
+typedef struct SymsConst 
+{
+  SymsTypeID type_id;
+  SymsStringRef name;
+  syms_uint value_len;
+  char value[SYMS_CONST_VALUE_MAX];
+} SymsConst;
+
+typedef char SymsConstIterImpl[128];
+typedef struct SymsConstIter 
+{
+  SymsDebugInfo *debug_info;
+  SymsConstIterImpl impl;
+} SymsConstIter;
+
+SYMS_API syms_bool
+syms_const_from_name(SymsInstance *instance, const char *name, syms_uint name_size, SymsConst *const_out);
+
+SYMS_API syms_bool
+  syms_const_iter_init(SymsInstance *instance, SymsConstIter *const_iter);
+
+SYMS_API syms_bool
+syms_const_iter_next(SymsConstIter *const_iter, SymsConst *const_out);
+
+// ----------------------------------------------------------------------------
+
+typedef char SymsSectionImpl[128];
+typedef struct SymsSection 
+{
+  SymsString name;
+  void *data;
+  SymsUMM data_size;
+  SymsAddr off;
+  SymsAddr va;
+  SymsSectionImpl impl;
+} SymsSection;
+
+typedef char SymsSecIterImpl[64];
+typedef struct SymsSecIter 
+{
+  SymsImage *image;
+  SymsSecIterImpl impl;
+} SymsSecIter;
+
+SYMS_API syms_bool
+syms_sec_iter_init(SymsInstance *instance, SymsSecIter *iter);
+// Initialize section iterator, on success return value is syms_true(1)
+//
+// On success returns "syms_true".
+
+SYMS_API syms_bool
+syms_sec_iter_next(SymsSecIter *iter, SymsSection *sec_out);
+// Parses a section out of image.
+// A shorthand for accessing sections of image by name.
+
+// ----------------------------------------------------------------------------
+
+typedef syms_uint SymsLn;
+typedef syms_uint SymsCol;
+typedef syms_uint SymsFileID;
+typedef syms_uint SymsSourceMapID;
+#define SYMS_INVALID_FILE_ID SYMS_INT_MAX
+
+typedef struct SymsSourceMap
+{
+  SymsLn ln;
+  SymsCol col;
+  SymsFileID file;
+} SymsSourceMap;
+
+typedef struct SymsAddrMap
+{
+  SymsAddr addr;
+  syms_uint id;
+} SymsAddrMap;
+
+typedef struct SymsLineTable
+{
+  syms_uint line_count;
+  syms_uint line_max;
+  syms_uint file_count;
+  SymsAddrMap *addrs;
+  SymsSourceMap *lines;
+  SymsStringRef *files;
+} SymsLineTable;
+
+typedef struct
+{
+  syms_uint line_id;
+  SymsSourceFileMap map;
+} SymsLineTableQuery;
+
+SYMS_API syms_bool
+syms_line_table_map_va(SymsLineTable *lt, SymsAddr va, SymsSourceFileMap *query_out);
+
+SYMS_API syms_bool
+syms_line_table_map_va_ex(SymsLineTable *lt, SymsAddr va, SymsLineTableQuery *query_out);
+
+SYMS_API syms_bool
+syms_line_table_map_src(SymsLineTable *lt, SymsInstance *instance, SymsString filename, syms_uint ln, SymsSourceFileMap *map_out);
+
+SYMS_API syms_bool
+syms_src_to_va(SymsInstance *instance, char *filename, syms_uint filename_len, syms_uint ln, SymsSourceFileMap *map_out);
+
+SYMS_API syms_bool
+syms_va_to_src(SymsInstance *instance, SymsAddr va, SymsSourceFileMap *map_out);
+
+// ----------------------------------------------------------------------------
+
+SYMS_API syms_bool
+syms_proc_from_name_2(SymsInstance *instance, const char *name, syms_uint name_size, SymsProc *proc_out, SymsModID *mod_out);
+// Searches given module for procedure that matches name.
+//
+// On success returns "syms_true".
+
+SYMS_API syms_bool
+syms_proc_from_name(SymsInstance *instance, const char *name, syms_uint name_size, SymsProc *proc_out);
+// Searches for all modules for procedure that matches name.
+//
+// On success returns "syms_true".
+
+SYMS_API syms_bool
+syms_proc_from_va(SymsInstance *instance, SymsAddr va, SymsProc *proc);
+
+SYMS_API syms_bool
+syms_proc_from_va_2(SymsInstance *instance, SymsAddr va, SymsProc *proc_out, SymsModID *mod_out);
+// Searches for nearest procedure symbol for specified address.
+//
+// On success returns "syms_true".
+
+SYMS_API syms_bool
+syms_img_sec_from_name(SymsInstance *instance, SymsString name, SymsSection *sec_out);
+
+SYMS_API syms_bool
+syms_va_to_secoff(SymsInstance *instance, SymsAddr va, syms_uint *isec_out, syms_uint *off_out);
+
+// ----------------------------------------------------------------------------
+
+SYMS_API SymsSymbol *
+syms_symbol_from_addr(SymsInstance *instance, SymsAddr addr, SymsArena *arena);
+// Queries symbols that are at or near specified address.
+// On successfull query result is a linked list of symbols.
+
+SYMS_API SymsSymbol *
+  syms_symbol_from_str(SymsInstance *instance, const char *str, SymsArena *arena);
+
+SYMS_API SymsSymbol *
+syms_symbol_from_str_ex(SymsInstance *instance, const char *str_ptr, syms_uint str_size, SymsArena *arena);
+// Queries symbol that match specifed string.
+// On successfull query result is a linked list of symbols.
+
+SYMS_API const char *
+syms_format_typeid(SymsInstance *instance, SymsTypeID *type_id, SymsArena *arena);
+// Creates a printable representation from type identifier.
+
+SYMS_API const char *
+syms_format_type(SymsInstance *instance, SymsType *type, SymsArena *arena);
+
+SYMS_API const char *
+syms_format_symbol(SymsInstance *instance, SymsSymbol *sym, SymsArena *arena);
+// Creates printable string for symbol.
+
+SYMS_API const char *
+syms_format_proc(SymsInstance *instance, SymsSymbolProc *proc, SymsArena *arena);
+
+SYMS_API const char *
+syms_format_method(SymsInstance *instance, SymsTypeMember *member, SymsArena *arena);
+
+SYMS_API SymsErrorCode
+syms_export_type(SymsInstance *instance, SymsTypeID type_id, SymsSymbolTypeInfo *type_out, SymsArena *arena);
+
+enum
+{
+  SYMS_BINDATA_TYPE_NULL,
+  SYMS_BINDATA_TYPE_NT32,
+  SYMS_BINDATA_TYPE_NT64,
+  SYMS_BINDATA_TYPE_ELF32,
+  SYMS_BINDATA_TYPE_ELF64,
+  SYMS_BINDATA_TYPE_DWARF,
+  SYMS_BINDATA_TYPE_PDB,
+  SYMS_BINDATA_TYPE_ELFSYM
+};
+typedef syms_uint syms_bindata_type;
+
+SYMS_API syms_bindata_type
+syms_examine_binary_data(void *buf, SymsUMM buf_size);
+
+//
+// Allocator
+//
+
+SYMS_API SymsArena *
+syms_borrow_memory(SymsInstance *instance);
+
+// TODO(nick): add parameter to configure virtual address reserve.
+#if 0
+syms_borrow_memory_ex(strcut SymsInstance *instance, SymsUMM reserve);
+#endif
+
+SYMS_API void
+syms_return_memory(SymsInstance *instance, SymsArena **arena);
+
+SYMS_API SymsUMM
+syms_get_arena_size(SymsArena *arena);
+
+//
+//  
+//
+
+SYMS_API SymsSymbolTypeInfo *
+syms_type_from_name_ex(SymsInstance *instance, const char *name, syms_uint name_size, SymsArena *arena);
+
+SYMS_API SymsSymbolStaticData *
+syms_global_from_name_ex(SymsInstance *instance, const char *name, syms_uint name_size, SymsArena *arena);
+
+SYMS_API SymsSymbolConstData *
+syms_const_from_name_ex(SymsInstance *instance, const char *name, syms_uint name_size, SymsArena *arena);
+
+SYMS_API SymsSymbolProc *
+syms_proc_from_addr_ex(SymsInstance *instance, SymsAddr va, SymsArena *arena);
+
+SYMS_API SymsSymbolProc *
+syms_proc_from_name_ex(SymsInstance *instance, const char *name, syms_uint name_size, SymsArena *arena);
+
+SYMS_API SymsSymbolSection *
+syms_img_sec_from_name_ex(SymsInstance *instance, const char *name, SymsArena *arena);
+
+SYMS_API SymsErrorCode
+syms_trampoline_from_ip(SymsInstance *instance, SymsAddr ip, SymsAddr *ip_out);
+
+//
+// Line table lookup
+//
+
+SYMS_API SymsSymbolSrcmap *
+syms_src_to_va_arena(SymsInstance *instance, const char *path, syms_uint path_size, syms_uint ln, syms_uint col, SymsArena *arena);
+
+SYMS_API SymsSymbolSrcmap *
+syms_va_to_src_arena(SymsInstance *instance, SymsAddr va, SymsArena *arena);
+
+//
+//
+//
+
+typedef struct SymsModInfo
+{
+  SymsMod header;
+  SymsBlockAllocator *procs;
+  SymsBlockAllocator *rangemap;
+  SymsLineTable       line_table;
+} SymsModInfo;
+
+SYMS_API SymsProc *
+syms_mod_info_find_proc(SymsModInfo *m, syms_uint i);
+
+SYMS_API SymsModInfo *
+syms_mod_from_va(SymsInstance *instance, SymsAddr va);
+
+//
+// Registers
+//
+
+enum
+{
+  SYMS_REG_null
+};
+
+// Each register can be assigned a class.
+// This provides a hint to the debugger as to how it might
+// want to interpret it's contents. 
+enum 
+{
+  SYMS_REG_CLASS_NULL,
+
+  // Basic processor state (IP/PC/flags) 
+  SYMS_REG_CLASS_STATE,
+
+  // Integer/pointer/segment/etc
+  SYMS_REG_CLASS_GPR,
+
+  // Control register
+  SYMS_REG_CLASS_CTRL,
+
+  // Scalar floating point register.
+  SYMS_REG_CLASS_FP,
+
+  // Vector register (could be N ints/floats/etc)
+  SYMS_REG_CLASS_VEC,
+
+  SYMS_REG_CLASS_INVALID
+};
+typedef syms_uint SymsRegClass;
+
+SYMS_API SymsTypeKind
+syms_get_reg_type(SymsArch arch, SymsRegID regid);
+
+SYMS_API const char *
+syms_get_reg_str(SymsArch arch, SymsRegID regid);
+
+SYMS_API syms_uint
+reg_desc_count_from_arch(SymsArch arch);
+
+SYMS_API SymsRegID
+syms_regs_get_first_regid_ex(SymsArch arch);
+
+SYMS_API SymsRegID
+syms_regs_get_last_regid_ex(SymsArch arch);
+
+SYMS_API SymsString
+syms_regs_get_name_ex(SymsArch arch, SymsRegID id);
+
+SYMS_API SymsTypeKind
+syms_regs_get_type_ex(SymsArch arch, SymsRegID id);
+
+SYMS_API SymsRegClass
+syms_regs_get_reg_class_ex(SymsArch arch, SymsRegID id);
+
+SYMS_API SymsRegID
+syms_regs_get_ip_ex(SymsArch arch);
+
+SYMS_API SymsRegID
+syms_regs_get_sp_ex(SymsArch arch);
+
+typedef struct
+{
+  SymsRegID aliasee;
+  syms_uint bit_shift;
+  syms_uint bit_count;
+} SymsRegAliasInfo;
+
+SYMS_API SymsRegAliasInfo
+syms_regs_get_alias_info(SymsArch arch, SymsRegID id);
+
+//
+//
+//
+
+SYMS_API const char *
+syms_member_access_to_str(SymsMemberAccess access);
+
+SYMS_API const char *
+syms_member_modifier_to_str(SymsMemberModifier modifier);
+
+SYMS_API const char *
+syms_type_modifier_to_str(SymsTypeModifier modifier);
+
+SYMS_API const char *
+syms_type_kind_to_str(SymsTypeKind kind);
+
+SYMS_API syms_uint
+syms_get_type_kind_bitcount(SymsArch arch, SymsTypeKind kind);
+
+// --------------------------------------------------------------------------------
+
+#ifdef __cplusplus
+}
+#endif
+
+#endif /* SYMS_PUBLIC_INCLUDE_H */
+
+
diff --git a/Engine/Source/Runtime/SymsLib/include/syms_virtual_unwind.h b/Engine/Source/Runtime/SymsLib/include/syms_virtual_unwind.h
new file mode 100644
--- /dev/null
+++ b/Engine/Source/Runtime/SymsLib/include/syms_virtual_unwind.h
@@ -0,0 +1,23 @@
+// Copyright Epic Games, Inc. All Rights Reserved.
+#ifndef SYMS_VIRTUAL_UNWIND_INCLUDE_H
+#define SYMS_VIRTUAL_UNWIND_INCLUDE_H
+
+/******************************************************************************
+ * File   : syms_virtual_unwind.h                                             *
+ * Author : Nikita Smith                                                      *
+ * Created: 2020/06/10                                                        *
+ * Purpose: User-level API for stack unwinding                                *
+ ******************************************************************************/
+
+typedef struct SymsVirtualUnwind SymsVirtualUnwind;
+
+SYMS_API SymsVirtualUnwind *
+syms_virtual_unwind_init(SymsInstance *instance, SymsArena *arena);
+
+SYMS_API SymsErrorCode
+syms_virtual_unwind_frame(SymsVirtualUnwind *context, 
+    void *memread_ctx, syms_memread_sig *memread_cb, 
+    void *regread_context, syms_regread_sig *regread_cb, 
+    void *regwrite_context, syms_regwrite_sig *regwrite_cb);
+
+#endif /* SYMS_VIRTUAL_UNWIND_INCLUDE_H */
diff --git a/Engine/Source/Runtime/SymsLib/include/symslib.h b/Engine/Source/Runtime/SymsLib/include/symslib.h
new file mode 100644
--- /dev/null
+++ b/Engine/Source/Runtime/SymsLib/include/symslib.h
@@ -0,0 +1,30 @@
+// Copyright Epic Games, Inc. All Rights Reserved.
+#ifndef SYMSLIB_INCLUDE_H
+#define SYMSLIB_INCLUDE_H
+
+#include "syms_public.h"
+#include "syms_core.h"
+#include "syms_block_alloc.h"
+
+#include "syms_platform.h"
+
+#include "regs/syms_regs.h"
+#include "regs/syms_regs_intel.h"
+
+#include "win/syms_nt.h"
+#include "elf/syms_elf.h"
+
+#include "pdb/syms_codeview.h"
+#include "pdb/syms_msf.h"
+#include "pdb/syms_dbi.h"
+#include "pdb/syms_pdb.h"
+#include "pdb/syms_pdb_api.h"
+
+#include "dwarf/syms_dwarf.h"
+#include "dwarf/syms_dwarf_api.h"
+#include "dwarf/syms_dwarf_unwind.h"
+
+#include "syms_virtual_unwind.h"
+#include "syms.h"
+
+#endif // SYMSLIB_INCLUDE_H
diff --git a/Engine/Source/Runtime/SymsLib/include/win/syms_nt.h b/Engine/Source/Runtime/SymsLib/include/win/syms_nt.h
new file mode 100644
--- /dev/null
+++ b/Engine/Source/Runtime/SymsLib/include/win/syms_nt.h
@@ -0,0 +1,598 @@
+// Copyright Epic Games, Inc. All Rights Reserved.
+#ifndef SYMS_WINDOWS_INCLUDE_H
+#define SYMS_WINDOWS_INCLUDE_H
+
+/******************************************************************************
+ * File   : syms_nt.h                                                         *
+ * Author : Nikita Smith                                                      *
+ * Created: 2020/05/30                                                        *
+ * Purpose: Declarations for NT executables                                   *
+ ******************************************************************************/
+
+#ifndef NT_API
+#define NT_API SYMS_API
+#endif
+
+typedef struct SymsGUID {
+  U32 data1;
+  U16 data2;
+  U16 data3;
+  U8  data4[8];
+} SymsGUID;
+
+SYMS_COMPILER_ASSERT(sizeof(SymsGUID) == 16);
+typedef struct SymsPeDebug {
+  U32 characteristics;
+  U32 time_stamp;
+  U16 major_ver;
+  U16 minor_ver;
+  U32 type;
+  U32 sizeof_data;
+  U32 raw_data_rva;
+  U32 raw_data_ptr;
+} SymsPeDebug;
+
+#pragma pack(push)
+
+#define SYMS_DOS_MAGIC 0x5a4d
+typedef struct SymsDosHeader {
+  U16 e_magic;                 /* Magic number 'M' 'Z' */
+  U16 e_cblp;                  /* Bytes on last page of file */
+  U16 e_cp;                    /* Pages in file */
+  U16 e_crlc;                  /* Relocations */
+  U16 e_cparhdr;               /* Size of header in paragraphs */
+  U16 e_minalloc;              /* Minimum extra paragraphs needed */
+  U16 e_maxalloc;              /* Maximum extra paragraphs needed */
+  U16 e_ss;                    /* Initial (relative) SS value */
+  U16 e_sp;                    /* Initial SP value */
+  U16 e_csum;                  /* Checksum */
+  U16 e_ip;                    /* Initial IP value */
+  U16 e_cs;                    /* Initial (relative) CS value */
+  U16 e_lfarlc;                /* File address of relocation table */
+  U16 e_ovno;                  /* Overlay number */
+  U16 e_res[4];                /* Reserved words */
+  U16 e_oemid;                 /* OEM identifier (for e_oeminfo) */
+  U16 e_oeminfo;               /* OEM information; e_oemid specific */
+  U16 e_res2[10];              /* Reserved words */
+  S32 e_lfanew;                /* File address of the PE header */
+} SymsDosHeader;
+
+typedef enum {
+  SYMS_NT_FILE_HEADER_MACHINE_UNKNOWN = 0x0,
+  SYMS_NT_FILE_HEADER_MACHINE_X86 = 0x14c,
+  SYMS_NT_FILE_HEADER_MACHINE_X64 = 0x8664,
+  SYMS_NT_FILE_HEADER_MACHINE_AM33 = 0x1d3,       // NOTE(nick): Matsushita AM33
+  SYMS_NT_FILE_HEADER_MACHINE_ARM = 0x1c0,
+  SYMS_NT_FILE_HEADER_MACHINE_ARM64 = 0xaa64,
+  SYMS_NT_FILE_HEADER_MACHINE_ARMNT = 0x1c4,
+  SYMS_NT_FILE_HEADER_MACHINE_EBC = 0xebc,
+  SYMS_NT_FILE_HEADER_MACHINE_I386 = 0x14c,
+  SYMS_NT_FILE_HEADER_MACHINE_IA64 = 0x200,
+  SYMS_NT_FILE_HEADER_MACHINE_M32R = 0x9041,      // NOTE(nick): Mitsubishi M32R little-endian
+  SYMS_NT_FILE_HEADER_MACHINE_MIPS16 = 0x266,
+  SYMS_NT_FILE_HEADER_MACHINE_MIPSFPU = 0x366,    // NOTE(nick): MIPS with FPU
+  SYMS_NT_FILE_HEADER_MACHINE_MIPSFPU16 = 0x466,  // NOTE(nick): MIPS16 with FPU
+  SYMS_NT_FILE_HEADER_MACHINE_POWERPC = 0x1f0,    // NOTE(nick): Power PC little-endian
+  SYMS_NT_FILE_HEADER_MACHINE_POWERPCFP = 0x1f1,  // NOTE(nick): Power PC with floating point support
+  SYMS_NT_FILE_HEADER_MACHINE_R4000 = 0x166, 
+  SYMS_NT_FILE_HEADER_MACHINE_RISCV32 = 0x5032,
+  SYMS_NT_FILE_HEADER_MACHINE_RISCV64 = 0x5064,
+  SYMS_NT_FILE_HEADER_MACHINE_RISCV128 = 0x5128,
+  SYMS_NT_FILE_HEADER_MACHINE_SH3 = 0x1a2,        // NOTE(nick): Hitachi SH3
+  SYMS_NT_FILE_HEADER_MACHINE_SH3DSP = 0x1a3,     // NOTE(nick): Hitachi SH3 DSP
+  SYMS_NT_FILE_HEADER_MACHINE_SH4 = 0x1a6,        // NOTE(nick): Hitachi SH4
+  SYMS_NT_FILE_HEADER_MACHINE_SH5 = 0x1a8,        // NOTE(nick): Hitachi SH5
+  SYMS_NT_FILE_HEADER_MACHINE_THUMB = 0x1c2,
+  SYMS_NT_FILE_HEADER_MACHINE_WCEMIPSV2 = 0x169   // NOTE(nick): MIPS little-endian WCE v2
+} SymsNTFileHeaderMachineType_e;
+typedef U16 SymsNTFileHeaderMachineType;
+
+typedef enum {
+  // NOTE(nick/MSDN): Image only, Windows CE, and Microsoft Windows NT and later. 
+  // This indicates that the file does not contain base relocations and must therefore be loaded at its preferred base address. 
+  // If the base address is not available, the loader reports an error. 
+  // The default behavior of the linker is to strip base relocations from executable (EXE) files. 
+  SYMS_NT_FILE_HEADER_FLAG_RELOCS_STRIPPED = 0x1,
+
+  // NOTE(nick/MSDN): Image only. This indicates that the image file is valid and can be run. 
+  // If this flag is not set, it indicates a linker error. 
+  SYMS_NT_FILE_HEADER_FLAG_EXECUTABLE_IMAGE = 0x2,
+
+  // NOTE(nick/MSDN): COFF line numbers have been removed. This flag is deprecated and should be zero.
+  SYMS_NT_FILE_HEADER_FLAG_LINE_NUMS_STRIPPED = 0x4,
+
+  // NOTE(nick/MSDN): COFF symbol table entries for local symbols have been removed. This flag is deprecated and should be zero.  
+  SYMS_NT_FILE_HEADER_FLAG_LOCAL_SYM_STRIPPED = 0x8,
+
+  // NOTE(nick/MSDN): Application can handle > 2-GB addresses. 
+  SYMS_NT_FILE_HEADER_FLAG_LARGE_ADDRESS_AWARE = 0x20,
+
+  // NOTE(nick/MSDN): Machine is based on a 32-bit-word architecture.
+  SYMS_NT_FILE_HEADER_FLAG_32BIT_MACHINE = 0x100,
+
+  // NOTE(nick/MSDN): Debugging information is removed from the image file. 
+  SYMS_NT_FILE_HEADER_FLAG_DEBUG_STRIPPED = 0x200,
+
+  // NOTE(nick/MSDN): If the image is on removable media, fully load it and copy it to the swap file. 
+  SYMS_NT_FILE_HEADER_FLAG_REMOVABLE_RUN_FROM_SWAP = 0x400,
+
+  // NOTE(nick/MSDN): If the image is on network media, fully load it and copy it to the swap file. 
+  SYMS_NT_FILE_HEADER_FLAG_NET_RUN_FROM_SWAP = 0x800,
+
+  // NOTE(nick/MSDN): The image file is a system file, not a user program. 
+  SYMS_NT_FILE_HEADER_FLAG_SYSTEM = 0x1000,
+
+  // NOTE(nick/MSDN): The image file is a dynamic-link library (DLL). 
+  // Such files are considered executable files for almost all purposes, 
+  // although they cannot be directly run. 
+  SYMS_NT_FILE_HEADER_FLAG_DLL = 0x2000,
+
+  // NOTE(nick/MSDN): The file should be run only on a uniprocessor machine. 
+  SYMS_NT_FILE_HEADER_FLAG_UP_SYSTEM_ONLY = 0x4000
+} SymsNTFileHeaderFlags_e;
+typedef U16 SymsNTFileHeaderFlags;
+
+typedef enum
+{
+  SYMS_NT_SUBSYSTEM_UNKNOWN = 0,
+  SYMS_NT_SUBSYSTEM_NATIVE = 1,
+  SYMS_NT_SUBSYSTEM_WINDOWS_GUI = 2,
+  SYMS_NT_SUBSYSTEM_WINDOWS_CUI = 3,
+  SYMS_NT_SUBSYSTEM_OS2_CUI = 5,
+  SYMS_NT_SUBSYSTEM_POSIX_CUI = 7,
+  SYMS_NT_SUBSYSTEM_NATIVE_WINDOWS = 8,
+  SYMS_NT_SUBSYSTEM_WINDOWS_CE_GUI = 9,
+  SYMS_NT_SUBSYSTEM_EFI_APPLICATION = 10,
+  SYMS_NT_SUBSYSTEM_EFI_BOOT_SERVICE_DRIVER = 11,
+  SYMS_NT_SUBSYSTEM_EFI_RUNTIME_DRIVER = 12,
+  SYMS_NT_SUBSYSTEM_EFI_ROM = 13,
+  SYMS_NT_SUBSYSTEM_XBOX = 14,
+  SYMS_NT_SUBSYSTEM_WINDOWS_BOOT_APPLICATION = 16
+} SymsNTSubsystem_e;
+typedef U16 SymsNTSubsystem;
+
+typedef enum
+{
+  SYMS_DLL_CHARACTERISTICS_HIGH_ENTROPY_VA       = 0x20,
+  SYMS_DLL_CHARACTERISTICS_DYNAMIC_BASE          = 0x40,
+  SYMS_DLL_CHARACTERISTICS_FORCE_INTEGRITY       = 0x80,
+  SYMS_DLL_CHARACTERISTICS_NX_COMPAT             = 0x100,
+  SYMS_DLL_CHARACTERISTICS_NO_ISOLATION          = 0x200,
+  SYMS_DLL_CHARACTERISTICS_NO_SEH                = 0x400,
+  SYMS_DLL_CHARACTERISTICS_NO_BIND               = 0x800,
+  SYMS_DLL_CHARACTERISTICS_APPCONTAINER          = 0x1000,
+  SYMS_DLL_CHARACTERISTICS_WDM_DRIVER            = 0x2000,
+  SYMS_DLL_CHARACTERISTICS_GUARD_CF              = 0x4000,
+  SYMS_DLL_CHARACTERISTICS_TERMINAL_SERVER_AWARE = 0x8000
+} SymsDllCharacteristics_e;
+
+typedef struct SymsNTSection {
+  char Name[8];
+
+  union {
+    U32 physical_addr;
+    U32 virtual_size;
+  } misc;
+
+  U32 virtual_addr;
+  U32 sizeof_raw_data;
+  U32 ptr_to_raw_data;
+  U32 ptr_to_relocs;
+  U32 ptr_to_linenumbers;
+  U16 num_relocs;
+  U16 num_lines;
+  U32 characteristics;
+} SymsNTSection;
+
+typedef enum {
+  SYMS_NT_DEBUG_DIR_UNKNOWN,
+  SYMS_NT_DEBUG_DIR_COFF,
+  SYMS_NT_DEBUG_DIR_CODEVIEW,
+  SYMS_NT_DEBUG_DIR_FPO,
+  SYMS_NT_DEBUG_DIR_MISC
+} SymsNTDebugDirType;
+
+typedef struct SymsNTDataDir {
+  U32 rva;
+  U32 len;
+} SymsNTDataDir;
+
+typedef struct SymsNTDebugDir {
+  U32 characteristics;
+  U32 time_stamp;
+  U16 major_ver;
+  U16 minor_ver;
+  U32 type;
+  U32 sizeof_data;
+  U32 raw_data_rva;
+  U32 raw_data_ptr;
+} SymsNTDebugDir;
+
+typedef enum {
+  SYMS_NT_DATA_DIR_EXPORT = 0,
+  SYMS_NT_DATA_DIR_IMPORT = 1,
+  SYMS_NT_DATA_DIR_RESOURCES = 2,
+  SYMS_NT_DATA_DIR_EXCEPTIONS = 3,
+  SYMS_NT_DATA_DIR_CERT = 4,
+  SYMS_NT_DATA_DIR_BASE_RELOC = 5,
+  SYMS_NT_DATA_DIR_DEBUG = 6,
+  SYMS_NT_DATA_DIR_ARCH = 7,
+  SYMS_NT_DATA_DIR_GLOBAL_PTR = 8,
+  SYMS_NT_DATA_DIR_TLS = 9,
+  SYMS_NT_DATA_DIR_LOAD_CONFIG = 10,
+  SYMS_NT_DATA_DIR_BOUND_IMPORT = 11,
+  SYMS_NT_DATA_DIR_IMPORT_ADDR = 12,
+  SYMS_NT_DATA_DIR_DELAY_IMPORT = 13,
+  SYMS_NT_DATA_DIR_COM_DESCRIPTOR = 14,
+  SYMS_NT_DATA_DIR_RESERVED = 15,
+
+  SYMS_NT_DATA_DIR_MAX
+} SymsNTDataDirType;
+
+typedef enum {
+  SYMS_CODEVIEW_SIG_V410  = 0x3930424e,
+  SYMS_CODEVIEW_SIG_V500  = 0x3131424e,
+  SYMS_CODEVIEW_SIG_PDB20 = 0x3031424e,
+  SYMS_CODEVIEW_SIG_PDB70 = 0x53445352
+} SymsCodeViewSig_e;
+typedef U32 SymsCodeViewSig;
+
+typedef struct SymsCodeViewHeaderPDB20 {
+  U32 sig;
+
+  // Offset in memory where debug info resides. If file is external this is 0
+  U32 off;
+
+  // Time when debug info was created (in seconds since 01.01.1970) 
+  U32 time;
+
+  // Initially set 1, and inceremented every time a PDB file is updated.
+  U32 age;
+
+  // Name of the file with the debug info 
+  // char name[];
+} SymsCodeViewHeaderPDB20;
+
+typedef struct SymsCodeViewHeaderPDB70 {
+  U32 sig;
+
+  // This GUID must match the one from the debug info file.
+  SymsGUID guid;
+
+  // Initially set 1, and inceremented every time a PDB file is updated.
+  U32 age;
+
+  // Name of the file with the debug info 
+  // char name[];
+} SymsCodeViewHeaderPDB70;
+
+#define SYMS_NT_FILE_HEADER_SIG 0x00004550u
+typedef struct SymsNTFileHeader {
+  SymsNTFileHeaderMachineType machine;
+  U16 number_of_sections;
+  U32 time_date_stamp;
+  U32 pointer_to_symbol_table;
+  U32 number_of_symbols;
+  U16 size_of_optional_header;
+  SymsNTFileHeaderFlags flags;
+} SymsNTFileHeader;
+
+typedef struct SymsNTOptionalHeader32 {
+  U16 magic;
+  U8 major_linker_version;
+  U8 minor_linker_version;
+  U32 sizeof_code;
+  U32 sizeof_inited_data;
+  U32 sizeof_uninited_data;
+  U32 entry_point_va;
+  U32 code_base;
+  U32 data_base;
+  U32 image_base;
+  U32 section_alignment;
+  U32 file_alignment;
+  U16 major_os_ver;
+  U16 minor_os_ver;
+  U16 major_img_ver;
+  U16 minor_img_ver;
+  U16 major_subsystem_ver;
+  U16 minor_subsystem_ver;
+  U32 win32_version_value;
+  U32 sizeof_image;
+  U32 sizeof_headers;
+  U32 check_sum;
+  U16 subsystem;
+  U16 dll_characteristics;
+  U32 sizeof_stack_reserve;
+  U32 sizeof_stack_commit;
+  U32 sizeof_heap_reserve;
+  U32 sizeof_heap_commit;
+  U32 loader_flags;
+  U32 number_of_rva_and_sizes;
+  SymsNTDataDir dirs[SYMS_NT_DATA_DIR_MAX];
+} SymsNTOptionalHeader32;
+
+typedef struct SymsNTOptionalHeader64 {
+  U16 magic;
+  U8 major_linker_version;
+  U8 minor_linker_version;
+  U32 sizeof_code;
+  U32 sizeof_inited_data;
+  U32 sizeof_uninited_data;
+  U32 entry_point_va;
+  U32 code_base;
+  U64 image_base;
+  U32 section_alignment;
+  U32 file_alignment;
+  U16 major_os_ver;
+  U16 minor_os_ver;
+  U16 major_img_ver;
+  U16 minor_img_ver;
+  U16 major_subsystem_ver;
+  U16 minor_subsystem_ver;
+  U32 win32_version_value;
+  U32 sizeof_image;
+  U32 sizeof_headers;
+  U32 check_sum;
+  U16 subsystem;
+  U16 dll_characteristics;
+  U64 sizeof_stack_reserve;
+  U64 sizeof_stack_commit;
+  U64 sizeof_heap_reserve;
+  U64 sizeof_heap_commit;
+  U32 loader_flags;
+  U32 number_of_rva_and_sizes;
+  SymsNTDataDir dirs[SYMS_NT_DATA_DIR_MAX];
+} SymsNTOptionalHeader64;
+
+typedef struct SymsNTImageHeader32 {
+  SymsNTFileHeader file_header;
+  SymsNTOptionalHeader32 opt_header;
+} SymsNTImageHeader32;
+
+typedef struct SymsNTImageHeader64 {
+  SymsNTFileHeader file_header;
+  SymsNTOptionalHeader64 opt_header;
+} SymsNTImageHeader64;
+
+#define SYMS_NT_IMAGE_SECTION_HEADER_CODE                0x20ul
+#define SYMS_NT_IMAGE_SECTION_HEADER_INITED_DATA         0x40ul
+#define SYMS_NT_IMAGE_SECTION_HEADER_UNINITED_DATA       0x80ul
+#define SYMS_NT_IMAGE_SECTION_HEADER_LNK_INFO            0x200ul
+#define SYMS_NT_IMAGE_SECTION_HEADER_LNK_REMOVE          0x800ul
+#define SYMS_NT_IMAGE_SECTION_HEADER_LNK_COMDAT          0x1000ul
+#define SYMS_NT_IMAGE_SECTION_HEADER_NO_DEFER_SPEC_EXC   0x4000ul
+#define SYMS_NT_IMAGE_SECTION_HEADER_GPREL               0x8000ul
+#define SYMS_NT_IMAGE_SECTION_HEADER_MEM_NOT_CACHED      0x04000000ul
+#define SYMS_NT_IMAGE_SECTION_HEADER_MEM_NOT_PAGED       0x08000000ul
+#define SYMS_NT_IMAGE_SECTION_HEADER_MEM_SHARED          0x10000000ul
+#define SYMS_NT_IMAGE_SECTION_HEADER_MEM_EXECUTE         0x20000000ul
+#define SYMS_NT_IMAGE_SECTION_HEADER_MEM_READ            0x40000000ul
+#define SYMS_NT_IMAGE_SECTION_HEADER_MEM_WRITE           0x80000000ul
+#define SYMS_NT_IMAGE_SECTION_HEADER_MEM_PURGEABLE       0x00020000ul
+#define SYMS_NT_IMAGE_SECTION_HEADER_MEM_LOCK            0x00040000ul
+#define SYMS_NT_IMAGE_SECTION_HEADER_MEM_PRELOAD         0x00080000ul
+#define SYMS_NT_IMAGE_SECTION_HEADER_DISCARDABLE         0x02000000ul
+#define SYMS_NT_IMAGE_SECTION_HEADER_LNK_NRELOC_OVFL     0x01000000ul
+#define SYMS_NT_IMAGE_SECTION_HEADER_ALIGN_1BYTES        0x00100000ul
+#define SYMS_NT_IMAGE_SECTION_HEADER_ALIGN_2BYTES        0x00200000ul
+#define SYMS_NT_IMAGE_SECTION_HEADER_ALIGN_4BYTES        0x00300000ul
+#define SYMS_NT_IMAGE_SECTION_HEADER_ALIGN_8BYTES        0x00400000ul
+#define SYMS_NT_IMAGE_SECTION_HEADER_ALIGN_16BYTES       0x00500000ul
+#define SYMS_NT_IMAGE_SECTION_HEADER_ALIGN_32BYTES       0x00600000ul
+#define SYMS_NT_IMAGE_SECTION_HEADER_ALIGN_64BYTES       0x00700000ul
+#define SYMS_NT_IMAGE_SECTION_HEADER_ALIGN_128BYTES      0x00800000ul
+#define SYMS_NT_IMAGE_SECTION_HEADER_ALIGN_256BYTES      0x00900000ul
+#define SYMS_NT_IMAGE_SECTION_HEADER_ALIGN_512BYTES      0x00A00000ul
+#define SYMS_NT_IMAGE_SECTION_HEADER_ALIGN_1024BYTES     0x00B00000ul
+#define SYMS_NT_IMAGE_SECTION_HEADER_ALIGN_2048BYTES     0x00C00000ul
+#define SYMS_NT_IMAGE_SECTION_HEADER_ALIGN_4096BYTES     0x00D00000ul
+#define SYMS_NT_IMAGE_SECTION_HEADER_ALIGN_8192BYTES     0x00E00000ul
+typedef u32 SymsNTImageSectionHeaderFlags;
+
+typedef U32 SymsNTImageSectionHeaderType;
+
+typedef struct SymsNTImageSectionHeader {
+  char name[8];
+
+  union {
+    U32 physical_address;
+    U32 virtual_size;
+  } u;
+
+  U32 va;
+  U32 sizeof_rawdata;
+  U32 rawdata_ptr;
+  U32 realloc_ptr;
+  U32 linenumbers_ptr;
+  U16 realloc_count;
+  U16 linenumbers_count;
+  U32 flags;
+} SymsNTImageSectionHeader;
+
+typedef struct SymsNTImageExportTable {
+  U32 characteristics;
+  U32 time_stamp;
+  U16 major_version;
+  U16 minor_version;
+  U32 name;
+  U32 base;
+  U32 num_funcs;
+  U32 num_names;
+  U32 funcs_rva;
+  U32 names_rva;
+  U32 ordinals_rva;
+} SymsNTImageExportTable;
+
+typedef struct SymsNTPdataPacked {
+  U32 rva_lo;
+  U32 rva_hi;
+  U32 uw_info_rva;
+} SymsNTPdataPacked;
+SYMS_COMPILER_ASSERT(sizeof(SymsNTPdataPacked) == 12);
+
+typedef struct SymsNTPdata {
+  SymsAddr lo;
+  SymsAddr hi;
+  SymsAddr uwinfo;
+} SymsNTPdata;
+
+// NOTE(nick): X(NAME, VALUE, NUM_CODES)
+#define SYMS_NT_UWOP_LIST \
+  X(SYMS_NT_OP_PUSH_NONVOL,      0, 1) \
+  X(SYMS_NT_OP_ALLOC_LARGE,      1, 2) \
+  X(SYMS_NT_OP_ALLOC_SMALL,      2, 1) \
+  X(SYMS_NT_OP_SET_FPREG,        3, 1) \
+  X(SYMS_NT_OP_SAVE_NONVOL,      4, 2) \
+  X(SYMS_NT_OP_SAVE_NONVOL_FAR,  5, 3) \
+  X(SYMS_NT_OP_EPILOG,           6, 2) \
+  X(SYMS_NT_OP_SPARE_CODE,       7, 3) \
+  X(SYMS_NT_OP_SAVE_XMM128,      8, 2) \
+  X(SYMS_NT_OP_SAVE_XMM128_FAR,  9, 3) \
+  X(SYMS_NT_OP_PUSH_MACHFRAME,  10, 1) \
+  /* end of the list */
+
+typedef enum {
+#define X(NAME, VALUE, NUM_CODES) NAME = VALUE,
+  SYMS_NT_UWOP_LIST
+#undef X
+  SYMS_NT_OP_COUNT
+} SymsNTUnwindOp;
+
+#define SYMS_NT_UNWIND_CODE_FLAGS_GET_OP_CODE(x) (((x) & 0x0f) >> 0x00)
+#define SYMS_NT_UNWIND_CODE_FLAGS_GET_OP_INFO(x) (((x) & 0xf0) >> 0x04)
+typedef union SymsNTUnwindCode {
+  struct {
+    U8 off;
+    U8 flags;
+  } u;
+  U16 frame_off;
+} SymsNTUnwindCode;
+SYMS_COMPILER_ASSERT(sizeof(SymsNTUnwindCode) == 2);
+
+enum {
+  // NOTE(nick): Describe type of exception handler, never figured out what these mean.
+  SYMS_NT_UNWIND_INFO_EHANDLER = (1 << 0),
+  SYMS_NT_UNWIND_INFO_UHANDLER = (1 << 1),
+  SYMS_NT_UNWIND_INFO_FHANDLER = SYMS_NT_UNWIND_INFO_EHANDLER|SYMS_NT_UNWIND_INFO_UHANDLER,
+
+  // NOTE(nick): Last code of unwind info is actually a SymsNTPdataPacked and it contains address to next unwind info.
+  SYMS_NT_UNWIND_INFO_CHAINED  = (1 << 2)
+};
+typedef U8 SymsNTUnwindInfoFlags;
+
+#define SYMS_NT_UNWIND_INFO_GET_CODE_COUNT(codes_num)    (((codes_num) + 1) & ~1)
+
+// NOTE(nick): Version of unwind info; from windows xp to windows 10 version is 1.
+#define SYMS_NT_UNWIND_INFO_HEADER_GET_VERSION(x)        (((x) & 0x07) >> 0x00)
+
+// NOTE(nick): Masks out flags, see SymsNTUnwindInfoFlags.
+#define SYMS_NT_UNWIND_INFO_HEADER_GET_FLAGS(x)          (((x) & 0xf8) >> 0x03)
+
+// NOTE(nick): Retrieves bits that indicate register kind.
+#define SYMS_NT_UWNIND_INFO_FRAME_GET_REG(x)             (((x) & 0x0f) >> 0x00)
+
+// NOTE(nick): Retrieves offset from register for stack frame.
+#define SYMS_NT_UNWIND_INFO_FRAME_GET_OFF(x)             (((x) & 0xf0) >> 0x04)
+
+typedef struct SymsNTUnwindInfo {
+  U8 header;
+
+  U8 prolog_size;
+
+  U8 codes_num;
+
+  U8 frame;
+
+  //SymsNTUnwindCode codes[0];
+} SymsNTUnwindInfo;
+SYMS_COMPILER_ASSERT(sizeof(SymsNTUnwindInfo) == 4);
+
+#define SYMS_NT_EXCEPTION_BREAKPOINT               0x80000003u
+#define SYMS_NT_EXCEPTION_SINGLE_STEP              0x80000004u
+#define SYMS_NT_EXCEPTION_LONG_JUMP                0x80000026u
+#define SYMS_NT_EXCEPTION_ACCESS_VIOLATION         0xC0000005u
+#define SYMS_NT_EXCEPTION_ARRAY_BOUNDS_EXCEEDED    0xC000008Cu
+#define SYMS_NT_EXCEPTION_DATA_TYPE_MISALIGNMENT   0x80000002u
+#define SYMS_NT_EXCEPTION_GUARD_PAGE_VIOLATION     0x80000001u
+#define SYMS_NT_EXCEPTION_FLT_DENORMAL_OPERAND     0xC000008Du
+#define SYMS_NT_EXCEPTION_FLT_DEVIDE_BY_ZERO       0xC000008Eu
+#define SYMS_NT_EXCEPTION_FLT_INEXACT_RESULT       0xC000008Fu
+#define SYMS_NT_EXCEPTION_FLT_INVALID_OPERATION    0xC0000090u
+#define SYMS_NT_EXCEPTION_FLT_OVERFLOW             0xC0000091u
+#define SYMS_NT_EXCEPTION_FLT_STACK_CHECK          0xC0000092u
+#define SYMS_NT_EXCEPTION_FLT_UNDERFLOW            0xC0000093u
+#define SYMS_NT_EXCEPTION_INT_DIVIDE_BY_ZERO       0xC0000094u
+#define SYMS_NT_EXCEPTION_INT_OVERFLOW             0xC0000095u
+#define SYMS_NT_EXCEPTION_PRIVILEGED_INSTRUCTION   0xC0000096u
+#define SYMS_NT_EXCEPTION_ILLEGAL_INSTRUCTION      0xC000001Du
+#define SYMS_NT_EXCEPTION_IN_PAGE_ERROR            0xC0000006u
+#define SYMS_NT_EXCEPTION_INVALID_DISPOSITION      0xC0000026u
+#define SYMS_NT_EXCEPTION_NONCONTINUABLE           0xC0000025u
+#define SYMS_NT_EXCEPTION_STACK_OVERFLOW           0xC00000FDu
+#define SYMS_NT_EXCEPTION_INVALID_HANDLE           0xC0000008u
+#define SYMS_NT_EXCEPTION_UNWIND_CONSOLIDATE       0x80000029u
+#define SYMS_NT_EXCEPTION_DLL_NOT_FOUND            0xC0000135u
+#define SYMS_NT_EXCEPTION_ORDINAL_NOT_FOUND        0xC0000138u
+#define SYMS_NT_EXCEPTION_ENTRY_POINT_NOT_FOUND    0xC0000139u
+#define SYMS_NT_EXCEPTION_DLL_INIT_FAILED          0xC0000142u
+#define SYMS_NT_EXCEPTION_CONTROL_C_EXIT           0xC000013Au
+#define SYMS_NT_EXCEPTION_FLT_MULTIPLE_FAULTS      0xC00002B4u
+#define SYMS_NT_EXCEPTION_FLT_MULTIPLE_TRAPS       0xC00002B5u
+#define SYMS_NT_EXCEPTION_NAT_CONSUMPTION          0xC00002C9u
+#define SYMS_NT_EXCEPTION_HEAP_CORRUPTION          0xC0000374u
+#define SYMS_NT_EXCEPTION_STACK_BUFFER_OVERRUN     0xC0000409u
+#define SYMS_NT_EXCEPTION_INVALID_CRUNTIME_PARAM   0xC0000417u
+#define SYMS_NT_EXCEPTION_ASSERT_FAILURE           0xC0000420u
+#define SYMS_NT_EXCEPTION_NO_MEMORY                0xC0000017u
+#define SYMS_VC_EXCEPTION_THROW                    0xE06D7363u
+
+#pragma pack(pop)
+
+typedef struct SymsImageNT {
+  SymsDosHeader *dos_header;
+  SymsNTFileHeader *file_header;
+  union {
+    SymsNTOptionalHeader32 *header32;
+    SymsNTOptionalHeader64 *header64;
+    void *header;
+  } u;
+  SymsString pdb_path;
+  U32 pdb_age;
+  U32 pdb_time;
+  SymsGUID pdb_guid;
+  U32 pdata_count;
+} SymsImageNT;
+
+typedef struct SymsSectIterNT {
+  struct SymsImage *img;
+  U32 header_index;
+  U32 header_count;
+  SymsNTImageSectionHeader *headers;
+} SymsSecIterNT;
+
+SYMS_INTERNAL SymsErrorCode
+syms_find_nearest_pdata(struct SymsInstance *instance, SymsAddr ip, SymsNTPdata *pdata_out);
+
+SYMS_INTERNAL SymsNTPdata
+syms_unpack_pdata(struct SymsInstance *instance, SymsNTPdataPacked *pdata);
+
+
+NT_API syms_bool
+syms_img_init_nt(struct SymsImage *img, void *img_data, SymsUWord img_size, SymsLoadImageFlags flags);
+
+NT_API const char *
+syms_get_nt_machine_str(u32 machine);
+
+SYMS_INTERNAL SymsSecIterNT
+syms_sec_iter_init_nt(struct SymsImage *img);
+
+SYMS_INTERNAL syms_bool
+syms_sec_iter_next_nt(SymsSecIterNT *iter, SymsNTImageSectionHeader *sec_out);
+
+SYMS_INTERNAL SymsNTPdata
+syms_unpack_pdata(struct SymsInstance *instance, struct SymsNTPdataPacked *pdata);
+
+SYMS_INTERNAL SymsAddr
+syms_get_rebase_nt(const SymsImageNT *nt, SymsImageHeaderClass header_class, SymsAddr base);
+
+#endif
diff --git a/Engine/Source/Runtime/SymsLib/include/win/syms_nt_unwind.h b/Engine/Source/Runtime/SymsLib/include/win/syms_nt_unwind.h
new file mode 100644
--- /dev/null
+++ b/Engine/Source/Runtime/SymsLib/include/win/syms_nt_unwind.h
@@ -0,0 +1,31 @@
+// Copyright Epic Games, Inc. All Rights Reserved.
+#ifndef SYMS_NT_UNWIND_INCLUDE_H
+#define SYMS_NT_UNWIND_INCLUDE_H
+
+/******************************************************************************
+ * File   : syms_nt_unwind.h                                                  *
+ * Author : Nikita Smith                                                      *
+ * Created: 2020/09/11                                                        *
+ * Purpose: Stack unwinder for Win32                                          *
+ ******************************************************************************/
+
+SYMS_INTERNAL syms_uint
+syms_nt_unwind_info_sizeof(SymsNTUnwindInfo *uwinfo);
+
+SYMS_INTERNAL syms_uint
+syms_nt_unwind_code_count_nodes(U8 uwcode_flags);
+
+SYMS_INTERNAL SymsRegID
+syms_remap_gpr_nt(SymsImage *img, syms_uint nt_regid);
+
+SYMS_INTERNAL SymsRegID
+syms_remap_xmm_nt(SymsImage *img, syms_uint nt_regid);
+
+SYMS_INTERNAL SymsErrorCode
+syms_memread_pdata(struct SymsInstance *instance, SymsMemread *memread_info, SymsAddr va, SymsNTPdata *pdata_out);
+
+SYMS_INTERNAL SymsErrorCode
+syms_virtual_unwind_nt(struct SymsInstance *instance, SymsRegs *regs, SymsMemread *memread_info);
+
+#endif // SYMS_NT_UNWIND_INCLUDE_H
+
