RAD symbol library

diff --git a/Engine/Source/Runtime/SymsLib/src/dwarf/syms_dwarf.c b/Engine/Source/Runtime/SymsLib/src/dwarf/syms_dwarf.c
new file mode 100644
--- /dev/null
+++ b/Engine/Source/Runtime/SymsLib/src/dwarf/syms_dwarf.c
@@ -0,0 +1,4655 @@
+// Copyright Epic Games, Inc. All Rights Reserved.
+SYMS_INTERNAL syms_bool
+dw_init_heap(DwContext *context, void *ud);
+
+SYMS_INTERNAL U32
+dw_trunc_u32(U64 value)
+{
+  U32 result;
+  SYMS_ASSERT_ALWAYS(value <= SYMS_UINT32_MAX);
+  result = (U32)value;
+  return result;
+}
+
+SYMS_INTERNAL DwBinRead
+dw_bin_read_init(DwMode mode, dw_uint addr_size, void *data, DwOffset max)
+{
+  DwBinRead result;
+  SYMS_ASSERT(addr_size <= 8);
+  result.err = syms_false;
+  result.addr_size = (U8)addr_size;
+  result.mode = mode;
+  result.off = 0;
+  result.max = max;
+  result.data = data;
+  return result;
+}
+
+SYMS_INTERNAL void *
+dw_bin_at(DwBinRead *bin)
+{
+  void *result = 0;
+  if (bin->off < bin->max) {
+    result = (void *)((U8 *)bin->data + bin->off);
+  }
+  return result;
+}
+
+SYMS_INTERNAL syms_bool
+dw_bin_seek(DwBinRead *bin, DwOffset off)
+{
+  syms_bool done = syms_false;
+  if (off <= bin->max) {
+    bin->off = off;
+    done = syms_true;
+  }
+  return done;
+}
+
+SYMS_INTERNAL syms_bool
+dw_bin_skip(DwBinRead *bin, DwOffset num)
+{
+  syms_bool done = syms_false;
+  if (bin->off + num <= bin->max) {
+    bin->off += num;
+    done = syms_true;
+  }
+  return done;
+}
+
+SYMS_INTERNAL syms_bool
+dw_bin_skip_s(DwBinRead *bin, S64 num)
+{
+  syms_bool done = syms_false;
+  S64 s;
+   
+  s = DW_CAST(S64, bin->off);
+  s += num;
+  if (s >= 0) {
+    U64 u = DW_CAST(U64, s);
+
+    if (u <= bin->max) {
+      bin->off = u;
+      done = syms_true;
+    }
+  }
+
+  return done;
+}
+
+SYMS_INTERNAL U8
+dw_bin_peek_u08(DwBinRead *bin)
+{
+  U8 result = 0;
+  if (bin->off + sizeof(result) <= bin->max) {
+    U8 *ptr = DW_CAST(U8 *, bin->data) + bin->off;
+    result = *ptr;
+  } else {
+    bin->err = syms_true;
+    SYMS_ASSERT_FAILURE("Out of bytes");
+  }
+  return result;
+}
+
+SYMS_INTERNAL U16
+dw_bin_peek_u16(DwBinRead *bin)
+{
+  U16 result = 0;
+  if (bin->off + sizeof(result) <= bin->max) {
+    U16 *ptr = (U16 *)dw_bin_at(bin);
+    result = *ptr;
+  } else {
+    bin->err = syms_true;
+    SYMS_ASSERT_FAILURE("Out of bytes");
+  }
+  return result;
+}
+
+SYMS_INTERNAL U32
+dw_bin_peek_u32(DwBinRead *bin)
+{
+  U32 result = 0;
+  if (bin->off + sizeof(result) <= bin->max) {
+    U32 *ptr = (U32 *)dw_bin_at(bin);
+    result = *(U32 *)ptr;
+  } else {
+    bin->err = syms_true;
+    SYMS_ASSERT_FAILURE("Out of bytes");
+  }
+  return result;
+}
+
+SYMS_INTERNAL U64
+dw_bin_peek_u64(DwBinRead *bin)
+{
+  U64 result = 0;
+  if (bin->off + sizeof(result) <= bin->max) {
+    U64 *ptr = (U64 *)dw_bin_at(bin);
+    result = *ptr;
+  } else {
+    bin->err = syms_true;
+    SYMS_ASSERT_FAILURE("Out of bytes");
+  }
+  return result;
+}
+
+SYMS_INTERNAL DwOffset
+dw_bin_peek_addr(DwBinRead *bin)
+{
+  DwOffset result = 0;
+  switch (bin->addr_size) {
+  case 1: result = dw_bin_peek_u08(bin); break;
+  case 2: result = dw_bin_peek_u16(bin); break;
+  case 4: result = dw_bin_peek_u32(bin); break;
+  case 8: result = dw_bin_peek_u64(bin); break;
+  default: SYMS_INVALID_CODE_PATH;
+  }
+  return result;
+}
+
+SYMS_INTERNAL syms_bool
+dw_bin_read(DwBinRead *bin, void *bf, dw_uint len)
+{
+  syms_bool done = syms_false;
+  if (bin->off + len <= bin->max) {
+    U8 *ptr = DW_CAST(U8 *, bin->data) + bin->off;
+    syms_memcpy(bf, ptr, len);
+    bin->off += len;
+    done = syms_true;
+  } else {
+    bin->err = syms_true;
+  }
+  return done;
+}
+
+SYMS_INTERNAL syms_bool
+dw_bin_subset(DwBinRead *bin, DwOffset offset, DwOffset size, DwBinRead *bin_out)
+{
+  syms_bool is_done = syms_false;
+
+  if (dw_bin_seek(bin, offset)) {
+    void *data = (void *)((U8 *)bin->data + offset);
+    if (dw_bin_seek(bin, offset + size)) {
+      *bin_out = dw_bin_read_init(bin->mode, bin->addr_size, data, size);
+      is_done = syms_true;
+    }
+  }
+
+  return is_done;
+}
+
+SYMS_INTERNAL U8
+dw_bin_read_u08(DwBinRead *bin)
+{
+  U8 result = 0;
+  dw_bin_read(bin, &result, sizeof(U8));
+  return result;
+}
+
+SYMS_INTERNAL U16
+dw_bin_read_u16(DwBinRead *bin)
+{
+  U16 result = 0;
+  dw_bin_read(bin, &result, sizeof(U16));
+  return result;
+}
+
+SYMS_INTERNAL U32
+dw_bin_read_u24(DwBinRead *bin)
+{
+  U32 result;
+  U32 a = DW_CAST(U32, dw_bin_read_u08(bin));
+  U32 b = DW_CAST(U32, dw_bin_read_u08(bin));
+  U32 c = DW_CAST(U32, dw_bin_read_u08(bin));
+  result = (a << 0) | (b << 8) | (c << 16);
+  return result;
+}
+
+SYMS_INTERNAL U32
+dw_bin_read_u32(DwBinRead *bin)
+{
+  U32 result = 0;
+  dw_bin_read(bin, &result, sizeof(U32));
+  return result;
+}
+
+SYMS_INTERNAL U64
+dw_bin_read_u64(DwBinRead *bin)
+{
+  U64 result = 0;
+  dw_bin_read(bin, &result, sizeof(U64));
+  return result;
+}
+
+SYMS_INTERNAL DwOffset
+dw_bin_read_length(DwBinRead *bin)
+{
+  DwOffset result = 0;
+  U32 r = dw_bin_read_u32(bin);
+  if (r < 0xfffffff0) {
+    result = DW_CAST(U64, r);
+  } else {
+    result = dw_bin_read_u64(bin);
+  }
+  return result;
+}
+
+SYMS_INTERNAL dw_uint
+dw_bin_read_uleb128(DwBinRead *bin)
+{
+  dw_uint res = 0;
+  dw_uint shift = 0;
+
+  U8 *data = DW_CAST(U8 *, bin->data) + bin->off;
+  U8 *start = data;
+
+  dw_uint num_used;
+
+  while (bin->off < bin->max) {
+    U8 i = *data++;
+    U8 val = i & 0x7fu;
+    res |= DW_CAST(U64, val) << shift;
+    if ((i & 0x80u) == 0) 
+      break;
+    shift += 7u;
+  }
+
+  num_used = DW_PTR_DIFF_BYTES(data, start);
+  if (num_used == 0) {
+    bin->err = syms_true;
+  }
+  bin->off += num_used;
+
+  return res;
+}
+
+SYMS_INTERNAL S8
+dw_bin_read_s08(DwBinRead *bin)
+{
+  S8 result = 0;
+  dw_bin_read(bin, &result, sizeof(S8));
+  return result;
+}
+
+SYMS_INTERNAL S16
+dw_bin_read_s16(DwBinRead *bin)
+{
+  S16 result = 0;
+  dw_bin_read(bin, &result, sizeof(S16));
+  return result;
+}
+
+SYMS_INTERNAL S32
+dw_bin_read_s32(DwBinRead *bin)
+{
+  S32 result = 0;
+  dw_bin_read(bin, &result, sizeof(S32));
+  return result;
+}
+
+SYMS_INTERNAL S64
+dw_bin_read_s64(DwBinRead *bin)
+{
+  S64 result = 0;
+  dw_bin_read(bin, &result, sizeof(S64));
+  return result;
+}
+
+SYMS_INTERNAL S64
+dw_bin_read_sleb128(DwBinRead *bin)
+{
+  S64 result = 0;
+  U32 shift = 0;
+
+  U8 *data = DW_CAST(U8 *, bin->data) + bin->off;
+  U8 *end = DW_CAST(U8 *, bin->data) + bin->max;
+  U8 *start = data;
+
+  while (data < end) {
+    U8 byte = *DW_CAST(U8 *, data++);
+    result |= DW_CAST(U64, (byte & 0x7fu)) << shift;
+    shift += 7u;
+    if ((byte & 0x80u) == 0u) {
+      DwOffset num_used;
+
+      if (shift < sizeof(result) * 8 && (byte & 0x40u) != 0)
+        result |= -(DW_CAST(S64, 1u << shift));
+
+      num_used = DW_PTR_DIFF_BYTES(data, start);
+      bin->off += num_used;
+
+      return DW_CAST(S64, result);
+    }
+  }
+  bin->err = syms_true;
+  return 0;
+}
+
+SYMS_INTERNAL SymsString
+dw_bin_read_string(DwBinRead *bin)
+{
+  // count chars
+  char *p_start = (char *)bin->data + bin->off;
+  char *p_end = (char *)bin->data + bin->max;
+  char *p_curr = p_start;
+  dw_uint byte_count;
+  SymsString str;
+
+  while (p_curr != p_end) {
+    if (*p_curr++ == '\0') {
+      break;
+    }
+  }
+  // pack string
+  byte_count = (dw_uint)(p_curr - p_start);
+  str = syms_string_init(p_start, byte_count);
+  // advance offset
+  bin->off += byte_count;
+  return str;
+}
+
+SYMS_INTERNAL DwOffset
+dw_bin_read_addr(DwBinRead *bin)
+{
+  DwOffset result = 0;
+  switch (bin->addr_size) {
+  case 1: result = dw_bin_read_u08(bin); break;
+  case 2: result = dw_bin_read_u16(bin); break;
+  case 4: result = dw_bin_read_u32(bin); break;
+  case 8: result = dw_bin_read_u64(bin); break;
+  default: SYMS_INVALID_CODE_PATH;
+  }
+  return result;
+}
+
+SYMS_INTERNAL DwOffset
+dw_bin_read_offset(DwBinRead *bin)
+{
+  DwOffset result = 0;
+  switch (bin->mode) {
+  case DW_MODE_32BIT: result = dw_bin_read_u32(bin); break;
+  case DW_MODE_64BIT: result = dw_bin_read_u64(bin); break;
+  default: SYMS_INVALID_CODE_PATH;
+  }
+  return result;
+}
+
+SYMS_INTERNAL DwBinRead
+dw_bin_read_for_sec(DwContext *context, DwMode mode, U8 addr_size, DwSecType type)
+{
+  DwImgSec *sec = &context->secs[type];
+  DwBinRead result;
+  if (sec) {
+    result = dw_bin_read_init(mode, addr_size, sec->data, sec->data_len);
+  } else {
+    result.off = 0;
+    result.max = 0;
+    result.data = 0;
+  }
+  return result;
+}
+
+SYMS_INTERNAL void
+dw_bin_write_init(DwBinWrite *write, void *bf, size_t max)
+{
+  write->off = 0;
+  write->max = max;
+  write->data = bf;
+}
+
+SYMS_INTERNAL syms_bool
+dw_bin_write(DwBinWrite *write, const void *data, size_t len)
+{
+  syms_bool done = syms_false;
+  if (write->off + len <= write->max) {
+    if (write->data) {
+      U8 *ptr = DW_CAST(U8 *, write->data) + write->off;
+      syms_memcpy(ptr, data, len);
+    }
+    done = syms_true;
+  }
+  write->off += len;
+  return done;
+}
+
+SYMS_INTERNAL syms_bool
+dw_bin_write_string(DwBinWrite *write, SymsString string)
+{
+  syms_bool done = dw_bin_write(write, string.data, string.len);
+  return done;
+}
+
+SYMS_INTERNAL U8
+dw_get_addr_size(SymsArch arch)
+{
+  switch (arch) {
+  case SYMS_ARCH_X86: return 4; break;
+  case SYMS_ARCH_X64: return 8; break;
+  default: return 0; break;
+  }
+}
+
+SYMS_INTERNAL DwTag
+dw_tag_bake_null(void)
+{
+  DwTag result;
+  result.cu     = 0;
+  result.info   = 0;
+  result.abbrev = 0;
+  return result;
+}
+
+SYMS_INTERNAL DwTag
+dw_tag_bake_with_abbrev(DwContext *context, DwOffset cu_info_base, DwOffset info_off, DwOffset abbrev_off)
+{
+  DwBinRead info_sec = dw_bin_read_for_sec(context, DW_MODE_NULL, 0, DW_SEC_INFO);
+  DwBinRead abbr_sec = dw_bin_read_for_sec(context, DW_MODE_NULL, 0, DW_SEC_ABBREV);
+  DwTag tag;
+
+  tag.cu     = 0;
+  tag.info   = 0;
+  tag.abbrev = 0;
+
+  if (dw_bin_seek(&info_sec, cu_info_base))
+    tag.cu = dw_bin_at(&info_sec);
+  if (dw_bin_seek(&info_sec, cu_info_base + info_off))
+    tag.info = dw_bin_at(&info_sec);
+  if (dw_bin_seek(&abbr_sec, abbrev_off))
+    tag.abbrev = dw_bin_at(&abbr_sec);
+
+  return tag;
+}
+
+SYMS_INTERNAL DwTag
+dw_tag_bake(DwContext *context, DwOffset cu_info_base, DwOffset info_off)
+{
+  DwTag result = dw_tag_bake_with_abbrev(context, cu_info_base, info_off, DW_INVALID_OFFSET);
+  return result;
+}
+
+SYMS_INTERNAL syms_bool
+dw_tag_is_valid(DwTag tag)
+{
+  syms_bool result = tag.cu != 0 && tag.info != 0;
+  return result;
+}
+
+SYMS_INTERNAL DwRef
+dw_invalid_ref(void)
+{
+  DwRef result;
+  result.info = DW_INVALID_OFFSET;
+  return result;
+}
+
+SYMS_INTERNAL syms_bool
+dw_ref_to_tag(DwContext *context, DwRef *ref, DwTag *tag_out)
+{
+  DwCuIter cu_iter;
+  if (dw_cu_iter_init(&cu_iter, context)) {
+    DwCompileUnit cu;
+    while (dw_cu_iter_next(&cu_iter, &cu)) {
+      if (cu.info_base >= ref->info && ref->info < cu.info_base + cu.info_len) {
+        *tag_out = dw_tag_bake_with_abbrev(context, cu.info_base, ref->info, DW_INVALID_OFFSET);
+        return syms_true;
+      }
+    }
+  }
+  return syms_false;
+}
+
+SYMS_INTERNAL void
+dw_seg_off_array_zero(DwSegOffArray *arr)
+{
+  arr->segoff_size = 0;
+  arr->segsel_size = 0;
+  arr->num = 0;
+  arr->entries = 0;
+}
+
+SYMS_INTERNAL syms_bool
+dw_seg_off_array_init(DwSegOffArray *arr, DwContext *context, DwMode mode, DwSecType sec_type, DwOffset base)
+{
+  DwBinRead sec = dw_bin_read_for_sec(context, mode, 0, sec_type);
+  syms_bool done = syms_false;
+
+  if (dw_bin_seek(&sec, base)) {
+    DwOffset unit_length = dw_bin_read_offset(&sec);
+    U16 ver = dw_bin_read_u16(&sec);
+    if (ver == 5) {
+      arr->segoff_size = dw_bin_read_u08(&sec);
+      arr->segsel_size = dw_bin_read_u08(&sec);
+      unit_length -= sizeof(U16) + sizeof(U8)*2;
+      arr->num = (dw_uint)(unit_length / (arr->segoff_size + arr->segsel_size));
+      arr->entries = dw_bin_at(&sec);
+      done = syms_true;
+    } else {
+      SYMS_ASSERT_FAILURE("invalid section");
+    }
+  }
+
+  if (!done) {
+    arr->segoff_size = 0;
+    arr->segsel_size = 0;
+    arr->num = 0;
+    arr->entries = 0;
+  }
+
+  return done;
+}
+
+SYMS_INTERNAL syms_bool
+dw_seg_off_array_get(DwSegOffArray *arr, dw_uint index, DwOffset *addr)
+{
+  dw_uint stride;
+  U8 *byte_ptr;
+  DwOffset seg, off;
+
+  if (index >= arr->num)
+    return syms_false;
+  if (!arr->entries)
+    return syms_false;
+  if (arr->segoff_size > sizeof(off) && arr->segsel_size > sizeof(seg))
+    return syms_false;
+
+  stride = arr->segsel_size + arr->segoff_size;
+  byte_ptr = DW_CAST(U8 *, arr->entries) + stride*index;
+
+  seg = 0;
+  syms_memcpy(&seg, byte_ptr, arr->segsel_size);
+  byte_ptr += arr->segsel_size;
+
+  off = 0;
+  syms_memcpy(&off, byte_ptr, arr->segoff_size);
+
+  *addr = seg + off;
+
+  return syms_true;
+}
+
+SYMS_INTERNAL void
+dw_off_array_zero(DwOffArray *arr)
+{
+  arr->entry_len = 0;
+  arr->num = 0;
+  arr->entries = 0;
+}
+
+SYMS_INTERNAL syms_bool
+dw_off_array_init(DwOffArray *arr, DwContext *context, DwMode mode, DwSecType sec_type, DwOffset base)
+{
+  DwBinRead sec = dw_bin_read_for_sec(context, mode, 0, sec_type);
+  syms_bool done = syms_false;
+
+  if (dw_bin_seek(&sec, base)) {
+    DwOffset unit_length = dw_bin_read_offset(&sec);
+    U16 ver = dw_bin_read_u16(&sec);
+    U16 padding = dw_bin_read_u16(&sec);
+
+    if (ver == 5) {
+      arr->entry_len = context->msize_byte_count;
+      unit_length -= sizeof(ver) + sizeof(padding);
+      arr->num = dw_trunc_u32(unit_length / context->msize_byte_count);
+      arr->entries = dw_bin_at(&sec);
+
+      done = syms_true;
+    } else {
+      SYMS_ASSERT_FAILURE("An invalid .debug_stroffsets section");
+    }
+  }
+
+  if (!done) {
+    arr->num = 0;
+    arr->entries = 0;
+  }
+
+  return done;
+}
+
+SYMS_INTERNAL syms_bool
+dw_off_array_get(DwOffArray *arr, U32 index, DwOffset *off)
+{
+  syms_bool done = syms_false;
+
+  if (index < arr->num) {
+    switch (arr->entry_len) {
+    case 4: {
+      U32 *offs_arr = DW_CAST(U32 *, arr->entries);
+      *off = DW_CAST(DwOffset, offs_arr[index]);
+      break;
+    }
+    case 8: {
+      U64 *offs_arr = DW_CAST(U64 *, arr->entries);
+      *off = DW_CAST(DwOffset, offs_arr[index]);
+      break;
+    }
+    default: SYMS_INVALID_CODE_PATH;
+    }
+
+    done = syms_true;
+  }
+
+  return done;
+}
+
+DW_API syms_bool
+dw_abbrev_iter_init(DwAbbrevIter *iter, DwContext *context, DwOffset abbrev_off)
+{
+  iter->data = dw_bin_read_for_sec(context, DW_MODE_NULL, 0, DW_SEC_ABBREV);
+  iter->state = DW_ABBREV_ITER_STATE_EMIT_DIE_BEGIN;
+  if (abbrev_off >= iter->data.max) {
+    return syms_false;
+  }
+  if (!dw_bin_seek(&iter->data, abbrev_off)) {
+    return syms_false;
+  }
+  return syms_true;
+}
+
+DW_API syms_bool
+dw_abbrev_iter_next(DwAbbrevIter *iter, DwAbbrevEntry *entry_out)
+{
+  entry_out->data_lo = iter->data.off;
+  switch (iter->state) {
+  case DW_ABBREV_ITER_STATE_NULL: return syms_false;
+  case DW_ABBREV_ITER_STATE_EMIT_DIE_BEGIN: {
+    entry_out->type = DW_ABBREV_ENTRY_TYPE_DIE_BEGIN;
+    iter->state = DW_ABBREV_ITER_STATE_EXPECT_TAG_INFO;
+  } break;
+  case DW_ABBREV_ITER_STATE_EXPECT_TAG_INFO: {
+    entry_out->type = DW_ABBREV_ENTRY_TYPE_TAG_INFO;
+    entry_out->u.tag_info.id = dw_bin_read_uleb128(&iter->data);
+    if (entry_out->u.tag_info.id == 0) {
+      entry_out->type = DW_ABBREV_ENTRY_TYPE_DIE_END;
+      if (iter->data.off < iter->data.max) {
+        iter->state = DW_ABBREV_ITER_STATE_EMIT_DIE_BEGIN;
+      } else {
+        iter->state = DW_ABBREV_ITER_STATE_NULL;
+      }
+      break;
+    }
+    entry_out->u.tag_info.tag = dw_bin_read_uleb128(&iter->data);
+    entry_out->u.tag_info.has_children = dw_bin_read_u08(&iter->data);
+    iter->state = DW_ABBREV_ITER_STATE_EXPECT_ATTRIB_INFO;
+  } break;
+  case DW_ABBREV_ITER_STATE_EXPECT_ATTRIB_INFO: {
+    entry_out->type = DW_ABBREV_ENTRY_TYPE_ATTRIB_INFO;
+    entry_out->u.attrib_info.name = dw_bin_read_uleb128(&iter->data);
+    entry_out->u.attrib_info.form = dw_bin_read_uleb128(&iter->data);
+    if (entry_out->u.attrib_info.form == DW_FORM_IMPLICIT_CONST) {
+      entry_out->u.attrib_info.has_implicit_const = syms_true;
+      entry_out->u.attrib_info.implicit_const = dw_bin_read_uleb128(&iter->data);
+    } else {
+      entry_out->u.attrib_info.has_implicit_const = syms_false;
+      entry_out->u.attrib_info.implicit_const = 0;
+    }
+    if (entry_out->u.attrib_info.name == 0 && entry_out->u.attrib_info.form == 0) {
+      entry_out->type = DW_ABBREV_ENTRY_TYPE_ATTRIB_INFO_NULL;
+      iter->state = DW_ABBREV_ITER_STATE_EXPECT_TAG_INFO;
+    }
+  } break;
+  }
+  entry_out->data_hi = iter->data.off;
+  if (iter->data.err) {
+    iter->state = DW_ABBREV_ITER_STATE_NULL;
+    entry_out->data_lo = DW_INVALID_OFFSET;
+    entry_out->data_hi = DW_INVALID_OFFSET;
+    entry_out->type = DW_ABBREV_ENTRY_TYPE_NULL;
+  }
+  return !iter->data.err;
+}
+
+typedef struct {
+  const char *name;
+  DwVersion ver;
+  U32 class_flags;
+} DwAttribInfo;
+
+SYMS_INTERNAL DwAttribInfo
+dw_attrib_get_info(DwAttribType attrib)
+{
+  DwAttribInfo result;
+  result.name = "";
+  result.ver = DWARF_INVALID_VERSION;
+  result.class_flags = 0;
+  switch (attrib) {
+#define X(id, val, v, cf) case val: result.name = #id; result.ver = v; result.class_flags = cf; break;
+    DW_ATTRIB_LIST
+#undef X
+  case DW_AT_HI_USER: break;
+  case DW_AT_LO_USER: break;
+  }
+  return result;
+}
+
+typedef struct {
+  const char *name;
+  DwVersion ver;
+  U32 class_flags;
+} DwFormInfo;
+
+SYMS_INTERNAL DwFormInfo
+dw_form_get_info(DwForm attrib)
+{
+  DwFormInfo result;
+  result.name = "";
+  result.ver = DWARF_INVALID_VERSION;
+  result.class_flags = 0;
+  switch (attrib) {
+#define X(id, val, v, cf) case val: result.name = #id; result.ver = (DwVersion)v; result.class_flags = cf; break;
+    DW_FORM_LIST
+#undef X
+  case DW_FORM_INVALID: break;
+  }
+  return result;
+}
+
+DW_API DwAttribClass
+dw_pick_attrib_value_class(DwCompileUnit *cu, DwAttribType attrib, DwForm form)
+{
+  DwAttribClass result;
+  U32 i;
+
+  DwAttribInfo attrib_info = dw_attrib_get_info(attrib);
+  DwFormInfo form_info = dw_form_get_info(form);
+
+  (void)cu;
+
+  /* NOTE(nick): Test for reserved attribute */
+  if (attrib_info.class_flags == 0 || form_info.class_flags == 0) {
+    return DW_AT_CLASS_INVALID;
+  }
+  result = DW_AT_CLASS_UNDEFINED;
+  for (i = 0; i < 32; ++i) {
+    U32 n = 1u << i;
+    if ((attrib_info.class_flags & n) != 0 && (form_info.class_flags & n) != 0) {
+      result = DW_CAST(DwAttribClass, n);
+      break;
+    }
+  }
+
+  return result;
+}
+
+SYMS_INTERNAL DwOffset
+dw_read_msize(DwContext *context, DwMSize *arr, U32 index)
+{
+  U8 *ptr = DW_CAST(U8 *, arr);
+  DwOffset result = 0;
+  ptr += context->msize_byte_count*index;
+  syms_memcpy(&result, ptr, context->msize_byte_count);
+  return result;
+}
+
+SYMS_INTERNAL syms_bool
+dw_attrib_iter_infer_address(DwAttribIter *iter, DwAttrib *attrib)
+{
+  syms_bool is_infered = syms_true;
+  if (attrib->form == DW_FORM_ADDR) {
+    attrib->value.address = attrib->form_value.addr;
+  } else if (attrib->form == DW_FORM_ADDRX  ||
+             attrib->form == DW_FORM_ADDRX1 ||
+             attrib->form == DW_FORM_ADDRX2 ||
+             attrib->form == DW_FORM_ADDRX3 ||
+             attrib->form == DW_FORM_ADDRX4) {
+    u32 addrx = syms_trunc_u32(attrib->form_value.addrx);
+    is_infered = dw_seg_off_array_get(&iter->cu->addrs_arr, addrx, &attrib->value.address);
+  } else {
+    SYMS_ASSERT_FAILURE_PARANOID("unable to infer address");
+    is_infered = syms_false;
+    attrib->value.address = SYMS_ADDR_MAX;
+  }
+  return is_infered;
+}
+
+SYMS_INTERNAL syms_bool
+dw_attrib_iter_infer_block(DwAttribIter *iter, DwAttrib *attrib)
+{ 
+  syms_bool is_infered = syms_true;
+  (void)iter;
+  if (attrib->form == DW_FORM_BLOCK  || 
+      attrib->form == DW_FORM_BLOCK1 ||
+      attrib->form == DW_FORM_BLOCK2 || 
+      attrib->form == DW_FORM_BLOCK4) {
+    attrib->value.block.len = attrib->form_value.block.len;
+    attrib->value.block.data = attrib->form_value.block.data;
+  } else {
+    SYMS_ASSERT_FAILURE_PARANOID("unable to infer block");
+    attrib->value.block.len = 0;
+    is_infered = syms_false;
+  }
+  return is_infered;
+}
+
+SYMS_INTERNAL syms_bool
+dw_attrib_iter_infer_const(DwAttribIter *iter, DwAttrib *attrib)
+{ 
+  syms_bool is_infered = syms_true;
+  (void)iter;
+  if (attrib->form == DW_FORM_DATA1 ||
+      attrib->form == DW_FORM_DATA2 ||
+      attrib->form == DW_FORM_DATA4 ||
+      attrib->form == DW_FORM_DATA8) {
+    // store data in the low 8 bytes
+    attrib->value.cnst16.lo = attrib->form_value.data;
+    attrib->value.cnst16.hi = 0;
+  } else if (attrib->form == DW_FORM_DATA16) {
+    attrib->value.cnst16.lo = attrib->form_value.data16.lo;
+    attrib->value.cnst16.hi = attrib->form_value.data16.hi;
+  } else if (attrib->form == DW_FORM_SDATA) {
+    attrib->value.cnst16.lo = (u64)attrib->form_value.sdata;
+    attrib->value.cnst16.hi = 0;
+  } else if (attrib->form == DW_FORM_UDATA) {
+    attrib->value.cnst16.lo = attrib->form_value.udata;
+    attrib->value.cnst16.hi = 0;
+  } else {
+    SYMS_ASSERT_FAILURE_PARANOID("unable to infer const");
+    attrib->value.cnst16.lo = 0;
+    attrib->value.cnst16.hi = 0;
+    is_infered = syms_false;
+  }
+  return is_infered;
+}
+
+SYMS_INTERNAL syms_bool
+dw_attrib_iter_infer_exprloc(DwAttribIter *iter, DwAttrib *attrib)
+{
+  syms_bool is_infered = syms_true;
+  (void)iter;
+  if (attrib->form == DW_FORM_EXPRLOC) {
+    SYMS_ASSERT(attrib->value_class == DW_AT_CLASS_EXPRLOC);
+    attrib->value.exprloc.len = attrib->form_value.exprloc.len;
+    attrib->value.exprloc.data = attrib->form_value.exprloc.data;
+  } else {
+    SYMS_ASSERT_FAILURE_PARANOID("unable to infer exprloc");
+    is_infered = syms_false;
+  }
+  return is_infered;
+}
+
+SYMS_INTERNAL syms_bool
+dw_attrib_iter_infer_flag(DwAttribIter *iter, DwAttrib *attrib)
+{
+  syms_bool is_infered = syms_false;
+  (void)iter;
+  if (attrib->form == DW_FORM_FLAG_PRESENT || attrib->form == DW_FORM_FLAG) {
+    SYMS_ASSERT(attrib->value_class == DW_AT_CLASS_FLAG);
+    attrib->value.flag = attrib->form_value.flag;
+    is_infered = syms_true;
+  } else {
+    SYMS_ASSERT_FAILURE_PARANOID("unable to infer flag");
+  }
+  return is_infered;
+}
+
+SYMS_INTERNAL syms_bool
+dw_attrib_iter_infer_addrptr(DwAttribIter *iter, DwAttrib *attrib)
+{
+  syms_bool done = syms_false;
+  (void)iter;
+  if (attrib->form == DW_FORM_SEC_OFFSET) {
+    SYMS_ASSERT(attrib->value_class == DW_AT_CLASS_ADDRPTR);
+    attrib->value.addrptr = attrib->form_value.sec_offset;
+    done = syms_true;
+  }
+  return done;
+}
+
+SYMS_INTERNAL syms_bool
+dw_attrib_iter_infer_lineptr(DwAttribIter *iter, DwAttrib *attrib)
+{
+  syms_bool done = syms_false;
+  (void)iter;
+  if (attrib->form == DW_FORM_SEC_OFFSET) {
+    SYMS_ASSERT(attrib->value_class == DW_AT_CLASS_LINEPTR);
+    attrib->value.loclistptr = attrib->form_value.sec_offset;
+    done = syms_true;
+  }
+  return done;
+}
+
+SYMS_INTERNAL syms_bool
+dw_attrib_iter_infer_loclist(DwAttribIter *iter, DwAttrib *attrib)
+{
+  syms_bool done = syms_false;
+  (void)iter;
+  if (attrib->form == DW_FORM_SEC_OFFSET) {
+    SYMS_ASSERT(attrib->value_class == DW_AT_CLASS_LOCLISTPTR);
+    attrib->value.loclistptr = attrib->form_value.sec_offset;
+    done = syms_true;
+  }
+  return done;
+}
+
+SYMS_INTERNAL syms_bool
+dw_attrib_iter_infer_loclistptr(DwAttribIter *iter, DwAttrib *attrib)
+{
+  syms_bool done = syms_false;
+  if (attrib->form == DW_FORM_LOCLISTX) {
+    DwCompileUnit *cu = iter->cu;
+    SYMS_ASSERT(attrib->value_class == DW_AT_CLASS_LOCLISTPTR);
+    done = dw_seg_off_array_get(&cu->loclists_arr, 
+                  dw_trunc_u32(attrib->form_value.loclistx), 
+                  &attrib->value.loclistptr);
+  } else if (attrib->form == DW_FORM_SEC_OFFSET) {
+    /* NOTE(nick): "loclist" class was added in the DWARF5, but
+     * why they duplicated already existed loclistptr?
+     * In particular the second version of this attribute,
+     * it even has same description as the "loclistptr" */
+    done = dw_attrib_iter_infer_loclist(iter, attrib);
+  }
+  return done;
+}
+
+SYMS_INTERNAL syms_bool
+dw_attrib_iter_infer_rnglistptr(DwAttribIter *iter, DwAttrib *attrib)
+{
+  syms_bool done = syms_false;
+
+  SYMS_ASSERT(attrib->value_class == DW_AT_CLASS_RNGLISTPTR);
+  if (attrib->form == DW_FORM_RNGLISTX) {
+    DwCompileUnit *cu = iter->cu;
+    SYMS_ASSERT(attrib->value_class == DW_AT_CLASS_RNGLISTPTR);
+    done = dw_seg_off_array_get(&cu->rnglists_arr, 
+                  dw_trunc_u32(attrib->form_value.rnglistx), 
+                  &attrib->value.rnglistptr);
+  } else if (attrib->form == DW_FORM_SEC_OFFSET) {
+    attrib->value.rnglistptr = attrib->form_value.sec_offset;
+    done = syms_true;
+  }
+  return done;
+}
+
+SYMS_INTERNAL syms_bool
+dw_attrib_iter_infer_rnglist(DwAttribIter *iter, DwAttrib *attrib)
+{
+  syms_bool done = syms_false;
+  if (attrib->form == DW_FORM_RNGLISTX) {
+
+  } else if (attrib->form == DW_FORM_SEC_OFFSET) {
+    done = dw_attrib_iter_infer_rnglistptr(iter, attrib);
+  }
+  return done;
+}
+
+SYMS_INTERNAL syms_bool
+dw_attrib_iter_infer_macptr(DwAttribIter *iter, DwAttrib *attrib)
+{
+  syms_bool is_infered = syms_true;
+  (void)iter;
+  if (attrib->form == DW_FORM_SEC_OFFSET) {
+    attrib->value.macptr = attrib->form_value.sec_offset;
+    SYMS_ASSERT(attrib->value_class == DW_AT_CLASS_MACPTR);
+  } else {
+    SYMS_ASSERT_FAILURE_PARANOID("unable to infer macptr");
+    attrib->value.macptr = DW_INVALID_OFFSET;
+    is_infered = syms_false;
+  }
+  return is_infered;
+}
+
+SYMS_INTERNAL syms_bool
+dw_attrib_iter_infer_ref(DwAttribIter *iter, DwAttrib *attrib)
+{
+  DwCompileUnit *cu = iter->cu;
+  DwContext *context = cu->dwarf;
+  syms_bool is_infered = syms_true;
+
+  SYMS_ASSERT_PARANOID(attrib->value_class == DW_AT_CLASS_REFERENCE);
+
+  if (attrib->form == DW_FORM_REF1 ||
+      attrib->form == DW_FORM_REF2 ||
+      attrib->form == DW_FORM_REF4 ||
+      attrib->form == DW_FORM_REF8) {
+    attrib->value.ref.info = cu->info_base + attrib->form_value.ref;
+  } else if (attrib->form == DW_FORM_REF_UDATA) {
+    attrib->value.ref.info = cu->info_base + attrib->form_value.udata;
+  } else if (attrib->form == DW_FORM_REF_ADDR) {
+    is_infered = syms_false;
+    if (iter->ref_addr_desc) {
+      if (context->next_info_ctx) {
+        /* TODO(nick): DWARF 5 */
+        SYMS_NOT_IMPLEMENTED;
+      } else {
+        attrib->value.ref.info = DW_INVALID_OFFSET;
+        is_infered = syms_true;
+      }
+    } else {
+      attrib->value.ref.info = attrib->form_value.ref;
+      is_infered = syms_true;
+    }
+  } else if (attrib->form == DW_FORM_REF_SUP4 ||
+              attrib->form == DW_FORM_REF_SUP8) {
+    attrib->value.ref.info = attrib->form_value.ref;
+  } else if (attrib->form ==  DW_FORM_REF_SIG8) {
+    /* TODO(nick): DWARF 5 
+     * We need to handle .debug_names section in order to resolve this value. */
+    SYMS_NOT_IMPLEMENTED;
+    is_infered = syms_false;
+  } else {
+    SYMS_ASSERT_FAILURE_PARANOID("unable to infer reference");
+    attrib->value.ref.info = DW_INVALID_OFFSET;
+    is_infered = syms_false;
+  }
+  return is_infered;
+}
+
+SYMS_INTERNAL syms_bool
+dw_attrib_iter_infer_string(DwAttribIter *iter, DwAttrib *attrib)
+{
+  DwCompileUnit *cu = iter->cu;
+  DwContext *context = cu->dwarf;
+  syms_bool is_infered = syms_true;
+  if (attrib->form == DW_FORM_STRING) {
+    attrib->value.string = attrib->form_value.string;
+  } else if (attrib->form == DW_FORM_STRX ||
+             attrib->form == DW_FORM_STRX1 ||
+             attrib->form == DW_FORM_STRX2 ||
+             attrib->form == DW_FORM_STRX3 ||
+             attrib->form == DW_FORM_STRX4) {
+    DwOffset strp = DW_INVALID_OFFSET;
+    u32 strx = dw_trunc_u32(attrib->form_value.strx);
+    dw_off_array_get(&cu->stroffs_arr, strx, &strp);
+    attrib->form_value.strp = strp;
+  } else if (attrib->form == DW_FORM_STRP ||
+             attrib->form == DW_FORM_STRP_SUP) {
+    { // Locate string in .debug_str section
+      DwBinRead sec = dw_bin_read_for_sec(context, DW_MODE_NULL, 0, DW_SEC_STR);
+      if (dw_bin_seek(&sec, attrib->form_value.strp)) {
+        attrib->value.string = dw_bin_read_string(&sec);
+      } else {
+        attrib->value.string = syms_string_init_lit("");
+      }
+    }
+  } else if (attrib->form == DW_FORM_LINE_STRP) {
+    DwBinRead sec = dw_bin_read_for_sec(context, DW_MODE_NULL, 0, DW_SEC_LINE_STR);
+    if (dw_bin_seek(&sec, attrib->form_value.strp)) {
+      attrib->value.string = dw_bin_read_string(&sec);
+    } else {
+      attrib->value.string = syms_string_init_lit("");
+    }
+  } else {
+    SYMS_ASSERT_FAILURE_PARANOID("unable to infer string");
+    attrib->value.string = syms_string_init_lit("");
+    is_infered = syms_false;
+  }
+  return is_infered;
+}
+
+SYMS_INTERNAL syms_bool
+dw_attrib_iter_infer_stroffptr(DwAttribIter *iter, DwAttrib *attrib)
+{
+  syms_bool is_infered = syms_true;
+  (void)iter;
+  if (attrib->form == DW_FORM_SEC_OFFSET) {
+    SYMS_ASSERT(attrib->value_class == DW_AT_CLASS_STRING);
+    attrib->value.stroffptr = attrib->form_value.sec_offset;
+  } else {
+    SYMS_ASSERT_FAILURE_PARANOID("unable to infer stroffptr");
+    attrib->value.stroffptr = DW_INVALID_OFFSET;
+    is_infered = syms_false;
+  }
+  return is_infered;
+}
+
+DW_API syms_bool
+dw_attrib_get_addr(DwAttrib *attrib, SymsAddr *addr_out)
+{
+  syms_bool is_value_valid = syms_true;
+  if (attrib->value_class == DW_AT_CLASS_ADDRESS) {
+    *addr_out = attrib->value.address;
+  } else {
+    *addr_out = 0;
+    is_value_valid = syms_false;
+  }
+  return is_value_valid;
+}
+
+DW_API syms_bool
+dw_attrib_get_block(DwAttrib *attrib, DwFormBlock *block_out)
+{
+  syms_bool done = syms_false;
+  if (attrib->value_class == DW_AT_CLASS_BLOCK) {
+    *block_out = attrib->value.block;
+    done = syms_true;
+  }
+  SYMS_ASSERT(done);
+  return done;
+}
+
+DW_API syms_bool
+dw_attrib_get_const32(DwAttrib *attrib, U32 *value)
+{
+  syms_bool done = syms_false;
+  if (attrib->value_class == DW_AT_CLASS_CONST && 
+      attrib->value.cnst16.hi == 0) {
+    *value = syms_trunc_u32(attrib->value.cnst16.lo);
+    done = syms_true;
+  }
+  SYMS_ASSERT(done);
+  return done;
+}
+
+DW_API syms_bool
+dw_attrib_get_const(DwAttrib *attrib, U64 *value)
+{
+  syms_bool done = syms_false;
+  if (attrib->value_class == DW_AT_CLASS_CONST && 
+      attrib->value.cnst16.hi == 0) {
+    *value = attrib->value.cnst16.lo;
+    done = syms_true;
+  }
+  SYMS_ASSERT(done);
+  return done;
+}
+
+DW_API syms_bool
+dw_attrib_get_const128(DwAttrib *attrib, U64 *lo, U64 *hi)
+{
+  syms_bool done = syms_false;
+  if (attrib->value_class == DW_AT_CLASS_CONST) {
+    *lo = attrib->value.cnst16.lo;
+    *hi = attrib->value.cnst16.hi;
+    done = syms_true;
+  }
+  SYMS_ASSERT(done);
+  return done;
+}
+
+
+SYMS_INTERNAL syms_bool
+dw_decode_location_expr(DwEncodedLocationExpr *encoded_loc, 
+            DwMode dw_mode, SymsArch arch,
+            void *memread_ctx, dw_memread_sig *memread,
+            void *regread_ctx, dw_regread_sig *regread,
+            DwLocation *loc)
+{
+#define DW_PUSHU(x)  *top = DW_CAST(U64, x); --top
+#define DW_PUSHS(x)  *top = DW_CAST(U64, x); --top
+#define DW_POP()     (++top)
+#define DW_POPS()   ((S64)*(++top))
+#define DW_POPU()   (*(++top))
+
+  U8 addr_size = dw_get_addr_size(arch);
+  syms_bool is_result_valid = syms_false;
+
+  U64 stack[128];
+  U64 *end = &stack[SYMS_ARRAY_SIZE(stack) - 1];
+  U64 *top = end;
+
+  DwBinRead opsmem = dw_bin_read_init(dw_mode, addr_size, encoded_loc->ops, encoded_loc->ops_size);
+
+  while (opsmem.off < opsmem.max) {
+    U8 op = dw_bin_read_u08(&opsmem);
+    U64 u, ua,ub;
+    S64 s, sa,sb;
+    U32 reg_index;
+
+    switch (op) {
+    case DW_OP_NOP: break;
+    case DW_OP_ADDR: {
+      u = dw_bin_read_addr(&opsmem); 
+      DW_PUSHU(u);
+    } break;
+    /* NOTE(nick): We could use the ... operator here, but
+     * that is the GNU extension and we want to be compiler agnostic. */
+    case DW_OP_LIT0:  case DW_OP_LIT1:  case DW_OP_LIT2:
+    case DW_OP_LIT3:  case DW_OP_LIT4:  case DW_OP_LIT5:
+    case DW_OP_LIT6:  case DW_OP_LIT7:  case DW_OP_LIT8:
+    case DW_OP_LIT9:  case DW_OP_LIT10: case DW_OP_LIT11:
+    case DW_OP_LIT12: case DW_OP_LIT13: case DW_OP_LIT14:
+    case DW_OP_LIT15: case DW_OP_LIT16: case DW_OP_LIT17:
+    case DW_OP_LIT18: case DW_OP_LIT19: case DW_OP_LIT20:
+    case DW_OP_LIT21: case DW_OP_LIT22: case DW_OP_LIT23:
+    case DW_OP_LIT24: case DW_OP_LIT25: case DW_OP_LIT26:
+    case DW_OP_LIT27: case DW_OP_LIT28: case DW_OP_LIT29:
+    case DW_OP_LIT30: case DW_OP_LIT31: {
+      u = DW_CAST(U64, op - DW_OP_LIT0);
+      DW_PUSHU(u);
+    } break;
+    case DW_OP_CONST1U: DW_PUSHU(1); break;
+    case DW_OP_CONST2U: DW_PUSHU(2); break;
+    case DW_OP_CONST4U: DW_PUSHU(4); break;
+    case DW_OP_CONST8U: DW_PUSHU(8); break;
+    case DW_OP_CONST1S: DW_PUSHS(1); break;
+    case DW_OP_CONST2S: DW_PUSHS(2); break;
+    case DW_OP_CONST4S: DW_PUSHS(4); break;
+    case DW_OP_CONST8S: DW_PUSHS(8); break;
+    case DW_OP_CONSTU: {
+      u = dw_bin_read_uleb128(&opsmem);
+      DW_PUSHU(u);
+    } break;
+    case DW_OP_CONSTS: {
+      s = dw_bin_read_sleb128(&opsmem);
+      DW_PUSHS(s);
+    } break;
+    case DW_OP_FBREG: {
+      s = dw_bin_read_sleb128(&opsmem);
+      s += encoded_loc->frame_base;
+      DW_PUSHS(s);
+    } break;
+    case DW_OP_BREG0:  case DW_OP_BREG1:  case DW_OP_BREG2:
+    case DW_OP_BREG3:  case DW_OP_BREG4:  case DW_OP_BREG5:
+    case DW_OP_BREG6:  case DW_OP_BREG7:  case DW_OP_BREG8:
+    case DW_OP_BREG9:  case DW_OP_BREG10: case DW_OP_BREG11:
+    case DW_OP_BREG12: case DW_OP_BREG13: case DW_OP_BREG14:
+    case DW_OP_BREG15: case DW_OP_BREG16: case DW_OP_BREG17:
+    case DW_OP_BREG18: case DW_OP_BREG19: case DW_OP_BREG20:
+    case DW_OP_BREG21: case DW_OP_BREG22: case DW_OP_BREG23:
+    case DW_OP_BREG24: case DW_OP_BREG25: case DW_OP_BREG26:
+    case DW_OP_BREG27: case DW_OP_BREG28: case DW_OP_BREG29:
+    case DW_OP_BREG30: case DW_OP_BREG31: {
+      reg_index = syms_trunc_u32(op - DW_OP_BREG0);
+      if (!regread(regread_ctx, arch, reg_index, &u, sizeof(u))) {
+        return syms_false;
+      }
+      s = dw_bin_read_sleb128(&opsmem);
+      s += DW_CAST(S64, u);
+      DW_PUSHS(s);
+    } break;
+    case DW_OP_BREGX: {
+      reg_index = syms_trunc_u32(dw_bin_read_uleb128(&opsmem));
+      if (!regread(regread_ctx, arch, reg_index, &u, sizeof(u))) {
+        return syms_false;
+      }
+
+      s = dw_bin_read_sleb128(&opsmem);
+      u = DW_CAST(U64, DW_CAST(S64, u) + s);
+      DW_PUSHU(u);
+    } break;
+    case DW_OP_DUP:     DW_PUSHU(*top);    break;
+    case DW_OP_DROP:    DW_POP();          break;
+    case DW_OP_OVER:    DW_PUSHU(top[1]);  break;
+    case DW_OP_PICK: {
+      U8 index = dw_bin_read_u08(&opsmem);
+      DW_PUSHU(top[index]);
+    } break;
+    case DW_OP_SWAP: {
+      u = top[0];
+      top[0] = top[1];
+      top[1] = u;
+    } break;
+    case DW_OP_ROT: {
+      u = top[0];
+      top[0] = top[2];
+      top[2] = top[1];
+      top[1] = u;
+    } break;
+    case DW_OP_DEREF: {
+      SYMS_ASSERT(sizeof(*top) <= addr_size);
+      u = DW_POPU();
+      if (!memread(memread_ctx, u, &u, addr_size)) {
+        return syms_false;
+      }
+      DW_PUSHU(u);
+    } break;
+    case DW_OP_DEREF_SIZE: {
+      U8 read_addr_size = dw_bin_read_u08(&opsmem);
+      if (read_addr_size > addr_size)
+        return syms_false;
+      u = DW_POPU();
+      if (!memread(memread_ctx, u, &u, read_addr_size)) {
+        return syms_false;
+      }
+      DW_PUSHU(u);
+    } break;
+    case DW_OP_XDEREF_SIZE:
+    case DW_OP_XDEREF: {
+      /* TODO(nick): This is a very rare case here, maybe later we can figure this out. */
+      return syms_false;
+    }
+    case DW_OP_PUSH_OBJECT_ADDRESS: {
+      DW_PUSHU(encoded_loc->member_location);
+    } break;
+    case DW_OP_FORM_TLS_ADDRESS: {
+      SYMS_NOT_IMPLEMENTED;
+    } break;
+    case DW_OP_CALL_FRAME_CFA: {
+      DW_PUSHU(encoded_loc->cfa);
+    } break;
+    case DW_OP_ABS: {
+      s = DW_POPS();
+      u = DW_CAST(U64, s < 0 ? -s : s);
+      DW_PUSHU(u);
+    } break;
+    case DW_OP_AND: {
+      sa = DW_POPS();
+      sb = DW_POPS();
+      s = sb & sa;
+      DW_PUSHS(s);
+    } break;
+    case DW_OP_DIV: {
+      sa = DW_POPS();
+      sb = DW_POPS();
+      if (sa == 0) {
+        return syms_false;
+      }
+      s = sb / sa;
+      DW_PUSHS(s);
+    } break;
+    case DW_OP_MINUS: {
+      sa = DW_POPS();
+      sb = DW_POPS();
+      s = sb - sa;
+      DW_PUSHS(s);
+    } break;
+    case DW_OP_MOD: {
+      sa = DW_POPS();
+      sb = DW_POPS();
+      s = sb % sa;
+      DW_PUSHS(s);
+    } break;
+    case DW_OP_MUL: {
+      sa = DW_POPS();
+      sb = DW_POPS();
+      if (sa == 0) {
+        return syms_false;
+      }
+      s = sb * sa;
+      DW_PUSHS(s);
+    } break;
+    case DW_OP_NEG: {
+      s = DW_POPS();
+      s = -s;
+      DW_PUSHS(s);
+    } break;
+    case DW_OP_NOT: {
+      u = DW_POPU();
+      u = DW_CAST(U64, !u);
+      DW_PUSHU(u);
+    } break;
+    case DW_OP_OR: {
+      ua = DW_POPU();
+      ub = DW_POPU();
+      u = ub | ua;
+      DW_PUSHU(u);
+    } break;
+    case DW_OP_PLUS: {
+      ua = DW_POPU();
+      ub = DW_POPU();
+      u = ub + ua;
+      DW_PUSHU(u);
+    } break;
+    case DW_OP_PLUS_UCONST: {
+      u = dw_bin_read_uleb128(&opsmem);
+      u += DW_POPU();
+      DW_PUSHU(u);
+    } break;
+    case DW_OP_SHL: {
+      ua = DW_POPU();
+      ub = DW_POPU();
+      u = ub << ua;
+      DW_PUSHU(u);
+    } break;
+    case DW_OP_SHR: {
+      ua = DW_POPU();
+      ub = DW_POPU();
+      u = ub >> ua;
+      DW_PUSHU(u);
+    } break;
+    case DW_OP_SHRA: {
+      sa = DW_POPS();
+      sb = DW_POPS();
+      s = sb >> sa;
+      DW_PUSHS(s);
+    } break;
+    case DW_OP_XOR: {
+      ua = DW_POPU();
+      ub = DW_POPU();
+      u = ub ^ ua;
+      DW_PUSHU(u);
+    } break;
+    case DW_OP_LE: {
+      sa = DW_POPS();
+      sb = DW_POPS();
+      s = sb <= sa;
+      DW_PUSHS(s);
+    } break;
+    case DW_OP_GE: {
+      sa = DW_POPS();
+      sb = DW_POPS();
+      s = sb >= sa;
+      DW_PUSHS(s);
+    } break;
+    case DW_OP_EQ: {
+      sa = DW_POPS();
+      sb = DW_POPS();
+      s = sb == sa;
+      DW_PUSHS(s);
+    } break;
+    case DW_OP_LT: {
+      sa = DW_POPS();
+      sb = DW_POPS();
+      s = sb < sa;
+      DW_PUSHS(s);
+    } break;
+    case DW_OP_GT: {
+      sa = DW_POPS();
+      sb = DW_POPS();
+      s = sb > sa;
+      DW_PUSHS(s);
+    } break;
+    case DW_OP_NE: {
+      sa = DW_POPS();
+      sb = DW_POPS();
+      s = sb != sa;
+      DW_PUSHS(s);
+    } break;
+    case DW_OP_SKIP: {
+      s = dw_bin_read_s16(&opsmem);
+      if (!dw_bin_skip_s(&opsmem, s)) {
+        return syms_false;
+      }
+    } break;
+    case DW_OP_BRA: {
+      s = dw_bin_read_s16(&opsmem);
+      sa = DW_POPS();
+      if (sa != 0) {
+        if (!dw_bin_skip_s(&opsmem, sa)) {
+          return syms_false;
+        }
+      }
+    } break;
+    case DW_OP_CALL_REF:
+    case DW_OP_CALL4:
+    case DW_OP_CALL2: {
+      SYMS_NOT_IMPLEMENTED;
+    } break;
+    case DW_OP_REG0:  case DW_OP_REG1:  case DW_OP_REG2:
+    case DW_OP_REG3:  case DW_OP_REG4:  case DW_OP_REG5:
+    case DW_OP_REG6:  case DW_OP_REG7:  case DW_OP_REG8:
+    case DW_OP_REG9:  case DW_OP_REG10: case DW_OP_REG11:
+    case DW_OP_REG12: case DW_OP_REG13: case DW_OP_REG14:
+    case DW_OP_REG15: case DW_OP_REG16: case DW_OP_REG17:
+    case DW_OP_REG18: case DW_OP_REG19: case DW_OP_REG20:
+    case DW_OP_REG21: case DW_OP_REG22: case DW_OP_REG23:
+    case DW_OP_REG24: case DW_OP_REG25: case DW_OP_REG26:
+    case DW_OP_REG27: case DW_OP_REG28: case DW_OP_REG29:
+    case DW_OP_REG30: case DW_OP_REG31: {
+      reg_index = DW_CAST(U32, op - DW_OP_REG0);
+      if (!regread(regread_ctx, arch, reg_index, &u, sizeof(u))) {
+        return syms_false;
+      }
+      DW_PUSHU(u);
+    } break;
+    case DW_OP_REGX: {
+      reg_index = syms_trunc_u32(dw_bin_read_uleb128(&opsmem));
+      if (!regread(regread_ctx, arch, reg_index, &u, sizeof(u))) {
+        return syms_false;
+      }
+
+      DW_PUSHU(u);
+    } break;
+    case DW_OP_IMPLICIT_VALUE: {
+      loc->type            = DW_LOCATION_IMPLICIT;
+      loc->u.implicit.len  = dw_bin_read_uleb128(&opsmem);
+      loc->u.implicit.data = dw_bin_at(&opsmem);
+      return syms_true;
+    }
+    case DW_OP_STACK_VALUE: {
+      goto exit;
+    }
+    case DW_OP_PIECE:
+    case DW_OP_BIT_PIECE: {
+      SYMS_NOT_IMPLEMENTED;
+    } break;
+    default: {
+      SYMS_ASSERT_FAILURE("encountered an unimplemented expression opcode");
+    } break;
+    }
+
+    is_result_valid = syms_true;
+  }
+
+exit:;
+
+   if (top != end && opsmem.off >= opsmem.max) {
+     loc->type = DW_LOCATION_ADDR;
+     loc->u.addr = DW_POPU();
+     is_result_valid = syms_true;
+   }
+
+   SYMS_ASSERT(is_result_valid);
+   return is_result_valid;
+
+}
+
+SYMS_INTERNAL syms_bool
+dw_decode_location_rva(DwEncodedLocationRva *encoded_loc, SymsAddr rebase, DwLocation *decoded_loc)
+{
+  decoded_loc->type = DW_LOCATION_ADDR;
+  decoded_loc->u.addr = rebase + encoded_loc->off;
+  return syms_true;
+}
+
+SYMS_API syms_bool
+dw_decode_location(DwEncodedLocation *encoded_loc, 
+           SymsAddr rebase,
+           void *memread_ctx, dw_memread_sig *memread,
+           void *regread_ctx, dw_regread_sig *regread,
+           DwLocation *decoded_loc)
+{
+  syms_bool is_result_valid = syms_false;
+  switch (encoded_loc->type) {
+  case DW_ENCODED_LOCATION_EXPR: {
+    DwEncodedLocationExpr *expr = &encoded_loc->u.expr;
+    DwContext *dwarf = encoded_loc->cu->dwarf;
+
+    is_result_valid = dw_decode_location_expr(expr, 
+        dwarf->mode, dwarf->arch,
+        memread_ctx, memread, 
+        regread_ctx, regread, 
+        decoded_loc); 
+  } break;
+
+  case DW_ENCODED_LOCATION_RVA: {
+    is_result_valid = dw_decode_location_rva(&encoded_loc->u.rva, rebase, decoded_loc); 
+  } break;
+
+  case DW_ENCODED_LOCATION_NULL: {
+    decoded_loc->type = DW_LOCATION_NULL;
+    is_result_valid = syms_true;
+  } break;
+  }
+  return is_result_valid;
+}
+
+SYMS_INTERNAL DwEncodedLocation
+dw_encode_null(void)
+{
+  DwEncodedLocation loc;
+  loc.type = DW_ENCODED_LOCATION_NULL;
+  loc.u.rva.cu = NULL;
+  loc.u.rva.off = 0;
+  return loc;
+}
+
+SYMS_INTERNAL DwEncodedLocation
+dw_encode_rva(DwCompileUnit *cu, U64 rva)
+{
+  DwEncodedLocation loc;
+  loc.type = DW_ENCODED_LOCATION_RVA;
+  loc.u.rva.cu = cu;
+  loc.u.rva.off = rva;
+  return loc;
+}
+
+SYMS_INTERNAL syms_bool
+dw_encoded_location_is_valid(DwEncodedLocation loc)
+{
+  syms_bool is_result_valid = syms_false;
+  switch (loc.type) {
+  case DW_ENCODED_LOCATION_NULL: {
+    is_result_valid = syms_true;
+  } break;
+
+  case DW_ENCODED_LOCATION_EXPR: {
+    is_result_valid = loc.cu != NULL && loc.u.expr.ops != NULL && loc.u.expr.ops_size > 0;
+  } break;
+
+  case DW_ENCODED_LOCATION_RVA: {
+    is_result_valid = loc.u.rva.cu != 0 && loc.u.rva.off != 0;
+  } break;
+  }
+  return is_result_valid;
+}
+
+#if 0
+SYMS_INTERNAL syms_bool
+dw_attrib_get_exprloc(DwAttrib *attrib, 
+    DwCompileUnit *cu, 
+    SymsAddr frame_base, 
+    SymsAddr member_location,
+    SymsAddr cfa,
+    DwLocation *loc)
+{
+  /* TODO(nick): Testing this without runtime is really hard, 
+   * so extensively test each case when lib is hooked up with debugger. */
+
+#define DW_PUSHU(x)  *top = DW_CAST(U64, x); --top;
+#define DW_POPU()    DW_CAST(U64, *(++top));
+
+#define DW_PUSHS(x) *top = DW_CAST(U64, x); --top;
+#define DW_POPS()   DW_CAST(S64, *(++top));
+
+  DwContext *dwarf = cu->dwarf;
+
+  syms_bool done = syms_false;
+
+  U64 stack[128];
+  U64 *end = &stack[SYMS_ARRAY_SIZE(stack) - 1];
+  U64 *top = end;
+
+  DwBinRead opsmem;
+
+  if (attrib->value_class != DW_AT_CLASS_EXPRLOC)
+    return syms_false;
+  
+
+  opsmem = dw_bin_read_init(dwarf->mode, cu->addr_size, 
+                attrib->form_value.exprloc.data, 
+                attrib->form_value.exprloc.len);
+
+  while (opsmem.off < opsmem.max) {
+    U8 op = dw_bin_read_u08(&opsmem);
+    U64 u;
+    S64 s;
+
+    U64 ua,ub;
+    S64 sa,sb;
+
+    U32 dwarf_reg_index;
+    U32 our_reg_index;
+
+    switch (op) {
+    case DW_OP_ADDR: {
+      u = dw_bin_read_addr(&opsmem); 
+      DW_PUSHU(u);
+      break;
+    }
+    /* NOTE(nick): We could use the ... operator here, but
+     * that is the GNU extension and we want to be compiler agnostic. */
+    case DW_OP_LIT0:  case DW_OP_LIT1:  case DW_OP_LIT2:
+    case DW_OP_LIT3:  case DW_OP_LIT4:  case DW_OP_LIT5:
+    case DW_OP_LIT6:  case DW_OP_LIT7:  case DW_OP_LIT8:
+    case DW_OP_LIT9:  case DW_OP_LIT10: case DW_OP_LIT11:
+    case DW_OP_LIT12: case DW_OP_LIT13: case DW_OP_LIT14:
+    case DW_OP_LIT15: case DW_OP_LIT16: case DW_OP_LIT17:
+    case DW_OP_LIT18: case DW_OP_LIT19: case DW_OP_LIT20:
+    case DW_OP_LIT21: case DW_OP_LIT22: case DW_OP_LIT23:
+    case DW_OP_LIT24: case DW_OP_LIT25: case DW_OP_LIT26:
+    case DW_OP_LIT27: case DW_OP_LIT28: case DW_OP_LIT29:
+    case DW_OP_LIT30: case DW_OP_LIT31: {
+      u = op - DW_OP_LIT0;
+      DW_PUSHU(u);
+      break;
+    }
+    case DW_OP_CONST1U: DW_PUSHU(1); break;
+    case DW_OP_CONST2U: DW_PUSHU(2); break;
+    case DW_OP_CONST4U: DW_PUSHU(4); break;
+    case DW_OP_CONST8U: DW_PUSHU(8); break;
+    case DW_OP_CONST1S: DW_PUSHS(1); break;
+    case DW_OP_CONST2S: DW_PUSHS(2); break;
+    case DW_OP_CONST4S: DW_PUSHS(4); break;
+    case DW_OP_CONST8S: DW_PUSHS(8); break;
+    case DW_OP_CONSTU: {
+      u = dw_bin_read_uleb128(&opsmem);
+      DW_PUSHU(u);
+      break;
+    }
+    case DW_OP_CONSTS: {
+      s = dw_bin_read_sleb128(&opsmem);
+      DW_PUSHS(s);
+      break;
+    }
+
+    case DW_OP_FBREG: {
+      s = dw_bin_read_sleb128(&opsmem);
+      s += frame_base;
+      DW_PUSHS(s);
+      break;
+    }
+    case DW_OP_BREG0:  case DW_OP_BREG1:  case DW_OP_BREG2:
+    case DW_OP_BREG3:  case DW_OP_BREG4:  case DW_OP_BREG5:
+    case DW_OP_BREG6:  case DW_OP_BREG7:  case DW_OP_BREG8:
+    case DW_OP_BREG9:  case DW_OP_BREG10: case DW_OP_BREG11:
+    case DW_OP_BREG12: case DW_OP_BREG13: case DW_OP_BREG14:
+    case DW_OP_BREG15: case DW_OP_BREG16: case DW_OP_BREG17:
+    case DW_OP_BREG18: case DW_OP_BREG19: case DW_OP_BREG20:
+    case DW_OP_BREG21: case DW_OP_BREG22: case DW_OP_BREG23:
+    case DW_OP_BREG24: case DW_OP_BREG25: case DW_OP_BREG26:
+    case DW_OP_BREG27: case DW_OP_BREG28: case DW_OP_BREG29:
+    case DW_OP_BREG30: case DW_OP_BREG31: {
+      reg_index = syms_trunc_u32(op - DW_OP_BREG0);
+      if (!dwarf->regread(dwarf->regread_ctx, dwarf->arch, reg_index, &u, sizeof(u))) {
+        return syms_false;
+      }
+
+      s = dw_bin_read_sleb128(&opsmem);
+      u = DW_CAST(U64, DW_CAST(S64, u) + s);
+      DW_PUSHU(u);
+
+      break;
+    }
+    case DW_OP_BREGX: {
+      reg_index = syms_trunc_u32(dw_bin_read_uleb128(&opsmem));
+      if (!dw_remap_reg(dwarf->arch, reg_index, &our_reg_index)) {
+        break;
+      }
+      if (!dwarf->regread(dwarf->regread_ctx, dwarf->arch, our_reg_index, &u, sizeof(u))) {
+        return syms_false;
+      }
+
+      s = dw_bin_read_sleb128(&opsmem);
+      u = DW_CAST(U64, DW_CAST(S64, u) + s);
+      DW_PUSHU(u);
+
+      break;
+    }
+
+    case DW_OP_DUP:     DW_PUSHU(*top);     break;
+    case DW_OP_DROP:    DW_POPU();          break;
+    case DW_OP_OVER:    DW_PUSHU(top[1]);  break;
+    case DW_OP_PICK: {
+      U8 index = dw_bin_read_u08(&opsmem);
+      DW_PUSHU(top[index]);
+      break;
+    }
+    case DW_OP_SWAP: {
+      u = top[0];
+      top[0] = top[1];
+      top[1] = u;
+      break;
+    }
+    case DW_OP_ROT: {
+      u = top[0];
+      top[0] = top[2];
+      top[2] = top[1];
+      top[1] = u;
+      break;
+    }
+    case DW_OP_DEREF: {
+      SYMS_ASSERT(sizeof(*top) <= cu->addr_size);
+      u = DW_POPU();
+      if (!dwarf->memread(dwarf->memread_ctx, u, &u, cu->addr_size)) 
+        return syms_false;
+      DW_PUSHU(u);
+      break;
+    }
+    case DW_OP_DEREF_SIZE: {
+      U8 addr_size = dw_bin_read_u08(&opsmem);
+      if (addr_size > cu->addr_size)
+        return syms_false;
+      u = DW_POPU();
+      if (!dwarf->memread(dwarf->memread_ctx, u, &u, addr_size))
+        return syms_false;
+      DW_PUSHU(u);
+      break;
+    }
+    case DW_OP_XDEREF_SIZE:
+    case DW_OP_XDEREF: {
+      /* TODO(nick): This is a very rare case here, maybe later we can figure this out. */
+      return syms_false;
+    }
+
+    case DW_OP_PUSH_OBJECT_ADDRESS: {
+      DW_PUSHU(member_location);
+      break;
+    }
+    case DW_OP_FORM_TLS_ADDRESS: {
+      SYMS_NOT_IMPLEMENTED;
+      break;
+    }
+    case DW_OP_CALL_FRAME_CFA: {
+      DW_PUSHU(cfa);
+      break;
+    }
+
+    case DW_OP_ABS: {
+      s = DW_POPS();
+      u = DW_CAST(U64, s < 0 ? -s : s);
+      DW_PUSHU(u);
+      break;
+    }
+    case DW_OP_AND: {
+      sa = DW_POPS();
+      sb = DW_POPS();
+      s = sa & sb;
+      DW_PUSHS(s);
+      break;
+    }
+    case DW_OP_DIV: {
+      sa = DW_POPS();
+      sb = DW_POPS();
+      if (sa == 0)
+        return syms_false;
+      s = sb/sa;
+      DW_PUSHS(s);
+      break;
+    }
+    case DW_OP_MINUS: {
+      sa = DW_POPS();
+      sb = DW_POPS();
+      s = sb-sa;
+      DW_PUSHS(s);
+      break;
+    }
+    case DW_OP_MOD: {
+      sa = DW_POPS();
+      sb = DW_POPS();
+      s = sb%sa;
+      DW_PUSHS(s);
+      break;
+    }
+    case DW_OP_MUL: {
+      sa = DW_POPS();
+      sb = DW_POPS();
+      if (sa == 0)
+        return syms_false;
+      s = sb%sa;
+      DW_PUSHS(s);
+      break;
+    }
+    case DW_OP_NEG: {
+      s = DW_POPS();
+      s = -s;
+      DW_PUSHS(s);
+      break;
+    }
+    case DW_OP_NOT: {
+      u = DW_POPU();
+      u = !u;
+      DW_PUSHU(u);
+      break;
+    }
+    case DW_OP_OR: {
+      ua = DW_POPU();
+      ub = DW_POPU();
+      u = ub|ua;
+      DW_PUSHU(u);
+      break;
+    }
+    case DW_OP_PLUS: {
+      ua = DW_POPU();
+      ub = DW_POPU();
+      u = ub+ub;
+      DW_PUSHU(u);
+      break;
+    }
+    case DW_OP_PLUS_UCONST: {
+      u = dw_bin_read_uleb128(&opsmem);
+      u += DW_POPU();
+      DW_PUSHU(u);
+      break;
+    }
+    case DW_OP_SHL: {
+      ua = DW_POPU();
+      ub = DW_POPU();
+      u = ub << ua;
+      DW_PUSHU(u);
+      break;
+    }
+    case DW_OP_SHR: {
+      ua = DW_POPU();
+      ub = DW_POPU();
+      u = ub >> ua;
+      DW_PUSHU(u);
+      break;
+    }
+    case DW_OP_SHRA: {
+      sa = DW_POPS();
+      sb = DW_POPS();
+      s = sb >> sa;
+      DW_PUSHS(s);
+      break;
+    }
+    case DW_OP_XOR: {
+      ua = DW_POPU();
+      ub = DW_POPU();
+      u = ub ^ ua;
+      DW_PUSHU(u);
+      break;
+    }
+    case DW_OP_LE: {
+      sa = DW_POPS();
+      sb = DW_POPS();
+      s = sb <= sa;
+      DW_PUSHS(s);
+      break;
+    }
+    case DW_OP_GE: {
+      sa = DW_POPS();
+      sb = DW_POPS();
+      s = sb >= sa;
+      DW_PUSHS(s);
+      break;
+    }
+    case DW_OP_EQ: {
+      sa = DW_POPS();
+      sb = DW_POPS();
+      s = sb == sa;
+      DW_PUSHS(s);
+      break;
+    }
+    case DW_OP_LT: {
+      sa = DW_POPS();
+      sb = DW_POPS();
+      s = sb < sa;
+      DW_PUSHS(s);
+      break;
+    }
+    case DW_OP_GT: {
+      sa = DW_POPS();
+      sb = DW_POPS();
+      s = sb > sa;
+      DW_PUSHS(s);
+      break;
+    }
+    case DW_OP_NE: {
+      sa = DW_POPS();
+      sb = DW_POPS();
+      s = sb != sa;
+      DW_PUSHS(s);
+      break;
+    }
+    case DW_OP_SKIP: {
+      s = dw_bin_read_s16(&opsmem);
+      if (!dw_bin_skip_s(&opsmem, s))
+        return syms_false;
+      break;
+    }
+    case DW_OP_BRA: {
+      s = dw_bin_read_s16(&opsmem);
+      sa = DW_POPS();
+      if (sa != 0) {
+        if (!dw_bin_skip_s(&opsmem, sa))
+          return syms_false;
+      }
+      break;
+    }
+    case DW_OP_CALL_REF:
+    case DW_OP_CALL4:
+    case DW_OP_CALL2: {
+      SYMS_NOT_IMPLEMENTED;
+      break;
+    }
+    case DW_OP_NOP: break;
+
+    case DW_OP_REG0:  case DW_OP_REG1:  case DW_OP_REG2:
+    case DW_OP_REG3:  case DW_OP_REG4:  case DW_OP_REG5:
+    case DW_OP_REG6:  case DW_OP_REG7:  case DW_OP_REG8:
+    case DW_OP_REG9:  case DW_OP_REG10: case DW_OP_REG11:
+    case DW_OP_REG12: case DW_OP_REG13: case DW_OP_REG14:
+    case DW_OP_REG15: case DW_OP_REG16: case DW_OP_REG17:
+    case DW_OP_REG18: case DW_OP_REG19: case DW_OP_REG20:
+    case DW_OP_REG21: case DW_OP_REG22: case DW_OP_REG23:
+    case DW_OP_REG24: case DW_OP_REG25: case DW_OP_REG26:
+    case DW_OP_REG27: case DW_OP_REG28: case DW_OP_REG29:
+    case DW_OP_REG30: case DW_OP_REG31: {
+      reg_index = DW_CAST(U32, op - DW_OP_REG0);
+      if (!dw_remap_reg(dwarf->arch, reg_index, &our_reg_index)) {
+        return syms_false;
+      }
+      if (!dwarf->regread(dwarf->regread_ctx, dwarf->arch, reg_index, &u, sizeof(u))) {
+        return syms_false;
+      }
+      DW_PUSHU(u);
+      break;
+    }
+    case DW_OP_REGX: {
+      reg_index = syms_trunc_u32(dw_bin_read_uleb128(&opsmem));
+      if (!dw_remap_reg(dwarf->arch, reg_index, &our_reg_index)) {
+        return syms_false;
+      }
+      if (!dwarf->regread(dwarf->regread_ctx, dwarf->arch, our_reg_index, &u, sizeof(u))) {
+        return syms_false;
+      }
+      DW_PUSHU(u);
+      break;
+    }
+    case DW_OP_IMPLICIT_VALUE: {
+      loc->type            = DW_LOCATION_IMPLICIT;
+      loc->u.implicit.len  = dw_bin_read_uleb128(&opsmem);
+      loc->u.implicit.data = dw_bin_at(&opsmem);
+      return syms_true;
+    }
+    case DW_OP_STACK_VALUE:
+      goto exit;
+
+    case DW_OP_PIECE:
+    case DW_OP_BIT_PIECE:
+      SYMS_NOT_IMPLEMENTED;
+    }
+
+    done = syms_true;
+  }
+
+exit:;
+
+   if (top != end && opsmem.off >= opsmem.max) {
+     loc->type = DW_LOCATION_ADDR;
+     loc->u.addr = DW_POPU();
+     done = syms_true;
+   }
+
+   SYMS_ASSERT(done);
+   return done;
+}
+#endif
+
+SYMS_INTERNAL syms_bool
+dw_attrib_get_flag(DwAttrib *attrib, U64 *flag)
+{
+  syms_bool done = syms_false;
+  if (attrib->value_class == DW_AT_CLASS_FLAG) {
+    *flag = attrib->value.flag;
+    done = syms_true;
+  }
+  SYMS_ASSERT(done);
+  return done;
+}
+
+SYMS_INTERNAL syms_bool
+dw_attrib_get_lineptr(DwAttrib *attrib, DwOffset *off)
+{
+  syms_bool done = syms_false;
+  if (attrib->value_class == DW_AT_CLASS_LINEPTR) {
+    *off = attrib->value.lineptr;
+    done = syms_true;
+  }
+  SYMS_ASSERT(done);
+  return done;
+}
+
+SYMS_INTERNAL syms_bool
+dw_attrib_get_loclistptr(DwAttrib *attrib, DwOffset *off)
+{
+  syms_bool done = syms_false;
+  if (attrib->value_class == DW_AT_CLASS_LOCLISTPTR) {
+    *off = attrib->value.loclistptr;
+    done = syms_true;
+  }
+  SYMS_ASSERT(done);
+  return done;
+}
+
+SYMS_INTERNAL syms_bool
+dw_attrib_get_macptr(DwAttrib *attrib, DwOffset *off)
+{
+  syms_bool done = syms_false;
+  if (attrib->value_class == DW_AT_CLASS_MACPTR) {
+    *off = attrib->value.macptr;
+    done = syms_true;
+  }
+  SYMS_ASSERT(done);
+  return done;
+}
+
+SYMS_INTERNAL syms_bool
+dw_attrib_get_rnglistptr(DwAttrib *attrib, DwOffset *off)
+{
+  syms_bool done = syms_false;
+  if (attrib->value_class == DW_AT_CLASS_RNGLISTPTR) {
+    *off = attrib->value.rnglistptr;
+    done = syms_true;
+  }
+  SYMS_ASSERT(done);
+  return done;
+}
+
+SYMS_INTERNAL syms_bool
+dw_attrib_get_rnglist(DwAttrib *attrib, DwOffset *off)
+{
+  syms_bool done = syms_false;
+  if (attrib->value_class == DW_AT_CLASS_RNGLIST || attrib->value_class == DW_AT_CLASS_RNGLISTPTR) {
+    *off = attrib->value.rnglist;
+    done = syms_true;
+  }
+  SYMS_ASSERT(done);
+  return done;
+}
+
+SYMS_INTERNAL syms_bool
+dw_attrib_get_addrptr(DwAttrib *attrib, DwOffset *off)
+{
+  syms_bool done = syms_false;
+  if (attrib->value_class == DW_AT_CLASS_ADDRPTR) {
+    *off = attrib->value.addrptr;
+    done = syms_true;
+  }
+  SYMS_ASSERT(done);
+  return done;
+}
+
+
+SYMS_INTERNAL syms_bool
+dw_attrib_get_ref(DwAttrib *attrib, DwRef *ref_out)
+{
+  syms_bool done = syms_false;
+  if (attrib->value_class == DW_AT_CLASS_REFERENCE) {
+    *ref_out = attrib->value.ref;
+    done = syms_true;
+  }
+  SYMS_ASSERT(done);
+  return done;
+}
+
+SYMS_INTERNAL syms_bool
+dw_attrib_get_ref_tag(DwAttrib *attrib, DwContext *context, DwTag *tag_out)
+{
+  DwRef ref;
+  if (dw_attrib_get_ref(attrib, &ref)) {
+    return dw_ref_to_tag(context, &ref, tag_out);
+  }
+  return syms_false;
+}
+
+SYMS_INTERNAL syms_bool
+dw_attrib_get_string(DwAttrib *attrib, SymsString *string)
+{
+  syms_bool done = syms_false;
+  if (attrib->value_class == DW_AT_CLASS_STRING) {
+    *string = attrib->value.string;
+    done = syms_true;
+  }
+  SYMS_ASSERT(done);
+  return done;
+}
+
+SYMS_INTERNAL syms_bool
+dw_attrib_get_stroffptr(DwAttrib *attrib, DwOffset *sec_off)
+{
+  syms_bool done = syms_false;
+  if (attrib->value_class == DW_AT_CLASS_STROFFSETSPTR) {
+    *sec_off = attrib->value.stroffptr;
+    done = syms_true;
+  }
+  SYMS_ASSERT(done);
+  return done;
+}
+
+SYMS_INTERNAL DwTag
+dw_attrib_iter_get_tag(DwAttribIter *iter)
+{
+  DwTag result;
+  if (iter->is_exhausted) {
+    result = dw_tag_bake_null();
+  } else {
+    result = dw_tag_bake_with_abbrev(
+        iter->cu->dwarf,
+        iter->cu->info_base,
+        iter->info_off, 
+        iter->abbrev_off);
+  }
+  return result;
+}
+
+SYMS_INTERNAL syms_bool
+dw_attrib_iter_init_from_tag(DwAttribIter *iter, DwContext *context, DwTag tag)
+{ (void)iter; (void)context; (void)tag;
+#if 0
+  U32 i;
+  syms_bool done = syms_false;
+  for (i = 0; i < dw->cu_num; ++i) {
+    DwCompileUnit *cu = &dw->cu_arr[i];
+    if (cu->info_data_start == tag.cu) {
+      DwOffset info_off = DW_PTR_DIFF_BYTES(tag.info, cu->info_data_start);
+      SYMS_ASSERT(info_off < cu->info_len);
+      done = dw_attrib_iter_init(iter, cu, info_off);
+      break;
+    }
+  }
+  return done;
+#else
+  SYMS_INVALID_CODE_PATH;
+  return syms_false;
+#endif
+}
+
+DW_API syms_bool
+dw_attrib_iter_init(DwAttribIter *iter, DwCompileUnit *cu, DwOffset info_off)
+{
+  iter->cu            = cu;
+  iter->is_exhausted  = syms_false;
+  iter->info_off      = info_off;
+  iter->ref_addr_desc = 0;
+  iter->info = dw_bin_read_for_sec(cu->dwarf, cu->mode, cu->addr_size, DW_SEC_INFO);
+  if (iter->info.max) {
+    SYMS_ASSERT(info_off >= cu->attribs_off);
+    iter->info.max = cu->info_base + cu->info_len;
+    if (dw_bin_seek(&iter->info, cu->info_base + info_off)) {
+      SymsOffset abbrev_off;
+
+      iter->abbrev_id = dw_bin_read_uleb128(&iter->info);
+      if (iter->info.err) {
+        return syms_false;
+      }
+      if (iter->abbrev_id == 0) {
+        iter->tag_type = DW_TAG_NULL;
+        iter->abbrev_header_len = 0;
+        iter->has_children = syms_false;
+        iter->is_exhausted = syms_true;
+        return syms_true; // inited as null
+      }
+      // TODO(nick): turn this into a function
+      abbrev_off = DW_INVALID_OFFSET;
+      {
+        DwAbbrevTable *table = &cu->abbrev_table;
+        if (table->entry_count > 0) {
+          S32 min = 0;
+          S32 max = (S32)table->entry_count - 1;
+          while (min <= max) {
+            S32 mid = (min + max) / 2;
+            if (iter->abbrev_id > table->entries[mid].id) {
+              min = mid + 1;
+            } else if (iter->abbrev_id < table->entries[mid].id) {
+              max = mid - 1;
+            } else {
+              abbrev_off = table->entries[mid].off;
+              break;
+            }
+          }
+        }
+      }
+      if (dw_abbrev_iter_init(&iter->abbrev_iter, cu->dwarf, abbrev_off)) {
+        DwAbbrevEntry abbrev;
+        if (dw_abbrev_iter_next(&iter->abbrev_iter, &abbrev)) {
+          SYMS_ASSERT(abbrev.type == DW_ABBREV_ENTRY_TYPE_DIE_BEGIN);
+          if (dw_abbrev_iter_next(&iter->abbrev_iter, &abbrev)) {
+            SYMS_ASSERT(abbrev.type == DW_ABBREV_ENTRY_TYPE_TAG_INFO);
+            if (abbrev.u.tag_info.id == iter->abbrev_id) {
+              DwAbbrevTagInfo *tag_info = &abbrev.u.tag_info;
+              iter->tag_type = (DwTagType)tag_info->tag;
+              iter->abbrev_off = abbrev_off;
+              iter->abbrev_header_len = abbrev.data_hi - abbrev.data_lo;
+              iter->has_children = tag_info->has_children;
+              return syms_true; // inited normally
+            }
+          }
+        }
+      }
+    }
+  }
+  return syms_false;
+}
+
+SYMS_INTERNAL syms_bool
+dw_attrib_iter_init2(DwAttribIter *iter, DwCompileUnit *cu, DwOffset info_off, DwOffset abbrev_off)
+{
+  iter->cu            = cu;
+  iter->is_exhausted  = syms_false;
+  iter->info_off      = info_off;
+  iter->ref_addr_desc = 0;
+  iter->info = dw_bin_read_for_sec(cu->dwarf, cu->mode, cu->addr_size, DW_SEC_INFO);
+  if (iter->info.max != 0) {
+    iter->info.max = cu->info_base + cu->info_len;
+    if (dw_bin_seek(&iter->info, cu->info_base + info_off)) {
+      iter->abbrev_id = dw_bin_read_uleb128(&iter->info);
+      if (iter->abbrev_id == 0) {
+        iter->tag_type = DW_TAG_NULL;
+        iter->abbrev_header_len = 0;
+        iter->has_children = syms_false;
+        iter->is_exhausted = syms_true;
+        return syms_true; // inited as null
+      }
+      if (dw_abbrev_iter_init(&iter->abbrev_iter, cu->dwarf, abbrev_off)) {
+        DwAbbrevEntry abbrev;
+        if (dw_abbrev_iter_next(&iter->abbrev_iter, &abbrev)) {
+          SYMS_ASSERT(abbrev.type == DW_ABBREV_ENTRY_TYPE_DIE_BEGIN);
+          if (dw_abbrev_iter_next(&iter->abbrev_iter, &abbrev)) {
+            SYMS_ASSERT(abbrev.type == DW_ABBREV_ENTRY_TYPE_TAG_INFO);
+            if (abbrev.u.tag_info.id == iter->abbrev_id) {
+              DwAbbrevTagInfo *tag_info = &abbrev.u.tag_info;
+              iter->tag_type = (DwTagType)tag_info->tag;
+              iter->abbrev_off = abbrev.data_lo;
+              iter->abbrev_header_len = abbrev.data_hi - abbrev.data_lo;
+              iter->has_children = tag_info->has_children;
+              return syms_true; // inited normally
+            }
+          }
+        }
+      }
+    }
+  }
+  return syms_false;
+}
+
+DW_API syms_bool
+dw_attrib_iter_reset(DwAttribIter *iter)
+{
+  return dw_attrib_iter_init2(iter, iter->cu, iter->info_off, iter->abbrev_off);
+}
+
+DW_API syms_bool
+dw_attrib_iter_next(DwAttribIter *iter, DwAttrib *attrib)
+{
+  DwAbbrevIter temp_abbrev_iter = iter->abbrev_iter;
+  DwAbbrevEntry abbrev;
+  DwAbbrevAttribInfo *attrib_info = 0;
+  DwBinRead *info = &iter->info;
+  syms_bool infered, done;
+
+  if (!iter->cu || iter->is_exhausted) {
+    return syms_false;
+  }
+
+  if (dw_abbrev_iter_next(&iter->abbrev_iter, &abbrev)) {
+    if (abbrev.type == DW_ABBREV_ENTRY_TYPE_ATTRIB_INFO) {
+      attrib_info = &abbrev.u.attrib_info;
+    }
+  }
+
+  if (!attrib_info) {
+    iter->abbrev_iter = temp_abbrev_iter;
+    iter->is_exhausted = syms_true;
+    return syms_false;
+  }
+  
+  attrib->info_lo = info->off;
+  attrib->info_hi = info->off;
+  attrib->tag = dw_tag_bake_with_abbrev(iter->cu->dwarf, iter->cu->info_base, iter->info_off, abbrev.data_lo);
+  attrib->name = (DwAttribType)attrib_info->name;
+  attrib->form = (DwForm)attrib_info->form;
+
+  /* NOTE(nick): This is a special case form. 
+   * Basically it let's user to define attribute form in the .debug_info. */
+  if (attrib->form == DW_FORM_INDIRECT) {
+    attrib->form = DW_CAST(DwForm, dw_bin_read_uleb128(info));
+  }
+  
+  switch (attrib->form) {
+  case DW_FORM_BLOCK1: {
+    attrib->form_value.block.len = (U64)dw_bin_read_u08(info);
+    attrib->form_value.block.data = dw_bin_at(info);
+    dw_bin_skip(info, attrib->form_value.block.len);
+  } break;
+  case DW_FORM_BLOCK2: {
+    attrib->form_value.block.len = (U64)dw_bin_read_u16(info);        
+    attrib->form_value.block.data = dw_bin_at(info);
+    dw_bin_skip(info, attrib->form_value.block.len);
+  } break;
+  case DW_FORM_BLOCK4: {
+    attrib->form_value.block.len = (U64)dw_bin_read_u32(info);        
+    attrib->form_value.block.data = dw_bin_at(info);
+    dw_bin_skip(info, attrib->form_value.block.len);
+  } break;
+  case DW_FORM_BLOCK: {
+    attrib->form_value.block.len = dw_bin_read_uleb128(info);
+    attrib->form_value.block.data = dw_bin_at(info);
+    dw_bin_skip(info, attrib->form_value.block.len);
+  } break;
+  case DW_FORM_DATA1:         attrib->form_value.data         = dw_bin_read_u08(info);        break;
+  case DW_FORM_DATA2:         attrib->form_value.data         = dw_bin_read_u16(info);        break;
+  case DW_FORM_DATA4:         attrib->form_value.data         = dw_bin_read_u32(info);        break;
+  case DW_FORM_DATA8:         attrib->form_value.data         = dw_bin_read_u64(info);        break;
+  case DW_FORM_UDATA:         attrib->form_value.udata        = dw_bin_read_uleb128(info);    break;
+  case DW_FORM_SDATA:         attrib->form_value.sdata        = dw_bin_read_sleb128(info);    break;
+  case DW_FORM_REF1:          attrib->form_value.ref          = dw_bin_read_u08(info);        break;
+  case DW_FORM_REF2:          attrib->form_value.ref          = dw_bin_read_u16(info);        break;
+  case DW_FORM_REF4:          attrib->form_value.ref          = dw_bin_read_u32(info);        break;
+  case DW_FORM_REF8:          attrib->form_value.ref          = dw_bin_read_u64(info);        break;
+  case DW_FORM_REF_SIG8:      attrib->form_value.ref          = dw_bin_read_u64(info);        break;
+  case DW_FORM_REF_ADDR:      attrib->form_value.ref          = dw_bin_read_offset(info);     break;
+  case DW_FORM_REF_UDATA:     attrib->form_value.ref          = dw_bin_read_uleb128(info);    break;
+  case DW_FORM_REF_SUP4:      attrib->form_value.ref          = dw_bin_read_u32(info);        break;
+  case DW_FORM_REF_SUP8:      attrib->form_value.ref          = dw_bin_read_u64(info);        break;
+  case DW_FORM_SEC_OFFSET:    attrib->form_value.sec_offset   = dw_bin_read_offset(info);     break;
+  case DW_FORM_ADDR:          attrib->form_value.addr         = dw_bin_read_addr(info);       break;
+  case DW_FORM_FLAG:          attrib->form_value.flag         = dw_bin_read_u08(info);        break;
+  case DW_FORM_FLAG_PRESENT:  attrib->form_value.flag         = 1;                            break;
+  case DW_FORM_STRP:          attrib->form_value.strp         = dw_bin_read_offset(info);     break;
+  case DW_FORM_STRP_SUP:      attrib->form_value.strp         = dw_bin_read_offset(info);     break;
+  case DW_FORM_LINE_STRP:     attrib->form_value.strp         = dw_bin_read_offset(info);     break;
+  case DW_FORM_STRING:        attrib->form_value.string       = dw_bin_read_string(info);     break;
+  case DW_FORM_STRX:          attrib->form_value.strx         = dw_bin_read_uleb128(info);    break;
+  case DW_FORM_STRX1:         attrib->form_value.strx         = dw_bin_read_u08(info);        break;
+  case DW_FORM_STRX2:         attrib->form_value.strx         = dw_bin_read_u16(info);        break;
+  case DW_FORM_STRX3:         attrib->form_value.strx         = dw_bin_read_u24(info);        break;
+  case DW_FORM_STRX4:         attrib->form_value.strx         = dw_bin_read_u32(info);        break;
+  case DW_FORM_ADDRX:         attrib->form_value.addrx        = dw_bin_read_uleb128(info);    break;
+  case DW_FORM_ADDRX1:        attrib->form_value.addrx        = dw_bin_read_u08(info);        break;
+  case DW_FORM_ADDRX2:        attrib->form_value.addrx        = dw_bin_read_u16(info);        break;
+  case DW_FORM_ADDRX3:        attrib->form_value.addrx        = dw_bin_read_u24(info);        break;
+  case DW_FORM_ADDRX4:        attrib->form_value.addrx        = dw_bin_read_u32(info);        break;
+  case DW_FORM_LOCLISTX:      attrib->form_value.loclistx     = dw_bin_read_uleb128(info);    break;
+  case DW_FORM_RNGLISTX:      attrib->form_value.rnglistx     = dw_bin_read_uleb128(info);    break;
+  case DW_FORM_DATA16: {
+    attrib->form_value.data16.lo = dw_bin_read_u64(info);
+    attrib->form_value.data16.hi = dw_bin_read_u64(info);
+  } break;
+  case DW_FORM_IMPLICIT_CONST: {
+    /* NOTE(nick): This is special case.
+     * Unlike other forms that have their values stored in the .debug_info section,
+     * this one defines it's value in the .debug_abbrev section. */
+    SYMS_ASSERT(attrib_info->has_implicit_const);
+    attrib->form_value.cnst = attrib_info->implicit_const;
+  } break;
+  case DW_FORM_EXPRLOC: {
+    attrib->form_value.exprloc.len = dw_bin_read_uleb128(info);
+    attrib->form_value.exprloc.data = DW_CAST(U8 *, info->data) + info->off;
+    if (!dw_bin_skip(info, attrib->form_value.exprloc.len)) {
+      return syms_false;
+    }
+  } break;
+  case DW_FORM_INDIRECT: {
+    SYMS_INVALID_CODE_PATH;
+  } break;
+  case DW_FORM_INVALID: {
+    return syms_false; 
+  } break;
+  }
+  attrib->info_hi = info->off;
+
+  infered = syms_false;
+  attrib->value_class = dw_pick_attrib_value_class(iter->cu, attrib->name, attrib->form);
+  switch (attrib->value_class) {
+  case DW_AT_CLASS_ADDRESS:       infered = dw_attrib_iter_infer_address(iter, attrib);      break;
+  case DW_AT_CLASS_ADDRPTR:       infered = dw_attrib_iter_infer_addrptr(iter, attrib);      break;
+  case DW_AT_CLASS_BLOCK:         infered = dw_attrib_iter_infer_block(iter, attrib);        break;
+  case DW_AT_CLASS_CONST:         infered = dw_attrib_iter_infer_const(iter, attrib);        break;
+  case DW_AT_CLASS_EXPRLOC:       infered = dw_attrib_iter_infer_exprloc(iter, attrib);      break;
+  case DW_AT_CLASS_FLAG:          infered = dw_attrib_iter_infer_flag(iter, attrib);         break;
+  case DW_AT_CLASS_LINEPTR:       infered = dw_attrib_iter_infer_lineptr(iter, attrib);      break;
+  case DW_AT_CLASS_LOCLIST:       infered = dw_attrib_iter_infer_loclist(iter, attrib);      break;
+  case DW_AT_CLASS_LOCLISTPTR:    infered = dw_attrib_iter_infer_loclistptr(iter, attrib);   break;
+  case DW_AT_CLASS_MACPTR:        infered = dw_attrib_iter_infer_macptr(iter, attrib);       break;
+  case DW_AT_CLASS_RNGLIST:       infered = dw_attrib_iter_infer_rnglist(iter, attrib);      break;
+  case DW_AT_CLASS_RNGLISTPTR:    infered = dw_attrib_iter_infer_rnglistptr(iter, attrib);   break;
+  case DW_AT_CLASS_REFERENCE:     infered = dw_attrib_iter_infer_ref(iter, attrib);          break;
+  case DW_AT_CLASS_STRING:        infered = dw_attrib_iter_infer_string(iter, attrib);       break;
+  case DW_AT_CLASS_STROFFSETSPTR: infered = dw_attrib_iter_infer_stroffptr(iter, attrib);    break;
+  case DW_AT_CLASS_INVALID: {
+    SYMS_ASSERT_FAILURE_PARANOID("attribute class was not resolved");
+  } break;
+  case DW_AT_CLASS_UNDEFINED: {
+    /* NOTE(nick): DWARF can contain attribs that aren't part of the spec. */
+    infered = syms_true; 
+  } break;
+  }
+
+  done = iter->info.err == 0 && infered;
+  return done;
+}
+
+SYMS_INTERNAL syms_bool
+dw_attrib_iter_next_tag(DwAttribIter *attrib_iter)
+{
+  DwAttrib attrib;
+  syms_bool done = syms_false;
+
+  while (dw_attrib_iter_next(attrib_iter, &attrib)) {
+    if (attrib_iter->is_exhausted) {
+      break;
+    }
+  }
+
+  if (attrib_iter->is_exhausted && !attrib_iter->info.err) {
+    SymsOffset abs_info_off = attrib_iter->info.off;
+    SymsOffset rel_info_off = abs_info_off - attrib_iter->cu->info_base;
+    done = dw_attrib_iter_init(attrib_iter, attrib_iter->cu, rel_info_off);
+  }
+
+  return done;
+}
+
+SYMS_INTERNAL syms_bool
+dw_attrib_iter_skip_children(DwAttribIter *iter)
+{
+  U32 depth = 0;
+  syms_bool done;
+  if (iter->has_children) {
+    do {
+      if (iter->has_children) {
+        ++depth;
+      }
+      if (iter->tag_type == DW_TAG_NULL) {
+        SYMS_ASSERT(depth > 0);
+        depth -= 1;
+        if (depth == 0) {
+          break;
+        }
+      }
+    } while (dw_attrib_iter_next_tag(iter));
+  }
+  done = (depth == 0);
+  return done;
+}
+
+DW_API syms_bool
+dw_tag_iter_init(DwTagIter *iter, DwCompileUnit *cu, SymsOffset info_off)
+{
+  iter->cu = cu;
+  iter->info_off = cu->info_base + info_off;
+  iter->depth = 0;
+  return syms_true;
+}
+
+DW_API syms_bool
+dw_tag_iter_begin(DwTagIter *iter, dw_uint *depth_out, DwTagType *tag_out, DwAttribIter *attribs_out)
+{
+  syms_bool is_result_valid = syms_false;
+  if (dw_attrib_iter_init(attribs_out, iter->cu, iter->info_off - iter->cu->info_base)) {
+    SYMS_ASSERT(iter->depth >= 0);
+    *depth_out = iter->depth;
+    *tag_out = attribs_out->tag_type;
+    if (attribs_out->tag_type == DW_TAG_NULL) {
+      iter->depth -= 1;
+    }
+    is_result_valid = syms_true;
+  }
+  return is_result_valid;
+}
+
+DW_API syms_bool
+dw_tag_iter_next(DwTagIter *iter, DwAttribIter *attribs)
+{
+  syms_bool is_ok;
+
+  if (!attribs->is_exhausted) {
+    DwAttrib dummy;
+    while (dw_attrib_iter_next(attribs, &dummy)) {
+    }
+  }
+  is_ok = syms_false;
+  if (attribs->is_exhausted) {
+    if (attribs->has_children) {
+      iter->depth += 1;
+    }
+    iter->info_off = attribs->info.off;
+    is_ok = syms_true;
+  }
+  return is_ok;
+}
+
+SYMS_INTERNAL syms_bool
+dw_name_iter_init(DwNameIter *iter, DwContext *context, DwNameTableIndex table_index)
+{
+  DwSecType sec_type = DW_SEC_NULL;
+  switch (table_index) {
+  case DW_NAME_TABLE_PUBTYPES: sec_type = DW_SEC_PUBTYPES; break;
+  case DW_NAME_TABLE_PUBNAMES: sec_type = DW_SEC_PUBNAMES; break;
+  case DW_NAME_TABLE_MAX: SYMS_INVALID_CODE_PATH;
+  }
+
+  iter->dwarf = context;
+  iter->sec = dw_bin_read_for_sec(context, DW_MODE_NULL, 0, sec_type);
+
+  iter->unit_start = DW_INVALID_OFFSET;
+  iter->unit_end = 0;
+  iter->unit_ver = 0;
+
+  iter->current_entry = 0;
+
+  iter->cu_info_off = DW_INVALID_OFFSET;
+  iter->cu_info_len = DW_INVALID_VALUE;
+
+  return syms_true;
+}
+
+SYMS_INTERNAL syms_bool
+dw_name_iter_next_table(DwNameIter *iter)
+{
+  syms_bool done = syms_false;
+
+  if (dw_bin_seek(&iter->sec, iter->unit_end)) {
+    iter->unit_start = iter->unit_end;
+    iter->unit_end   = dw_bin_read_length(&iter->sec);
+
+    if (iter->unit_end) {
+      iter->unit_end += iter->sec.off;
+
+      iter->unit_ver    = dw_bin_read_u16(&iter->sec);
+      iter->cu_info_off = dw_bin_read_offset(&iter->sec);
+      iter->cu_info_len = dw_bin_read_length(&iter->sec);
+
+      if (iter->unit_ver > 4 || iter->unit_ver < 1)
+        done = syms_false;
+
+      /* TODO(nick): Make sure that unit_info_off & unit_info_len are inside the 
+       * .debug_info section, if not fail. */
+      done = syms_true;
+    } else {
+      iter->unit_start = DW_INVALID_OFFSET;
+      iter->unit_end   = DW_INVALID_OFFSET;
+    }
+  }
+
+  return done;
+}
+
+SYMS_INTERNAL syms_bool
+dw_name_iter_next(DwNameIter *iter, DwTag *tag, SymsString *tag_name)
+{
+  syms_bool done = syms_false;
+
+redo:;
+  iter->current_entry = dw_bin_at(&iter->sec);
+
+  if (iter->sec.off < iter->unit_end) {
+    DwOffset info_off = dw_bin_read_offset(&iter->sec);
+    if (info_off) {
+      *tag = dw_tag_bake(iter->dwarf, iter->cu_info_off, info_off);
+      *tag_name = dw_bin_read_string(&iter->sec);
+      done = syms_true;
+    } else {
+      iter->current_entry = 0;
+    }
+  }
+
+  if (!done && dw_name_iter_next_table(iter))
+    goto redo;
+
+  return done;
+}
+
+SYMS_INTERNAL dw_uint
+dw_name_table_find(DwContext *context, DwNameTableIndex table_index, SymsString name, dw_uint tags_max, DwTag tags[])
+{
+  dw_uint tags_num = 0;
+  DwNameTable *table = context->name_tables[table_index];
+  if (table) {
+    dw_uint hash = syms_hash_djb2(name.data, name.len) % SYMS_ARRAY_SIZE(table->keys);
+    DwNameTableKeyValue *key;
+    for (key = table->keys[hash]; key != 0; key = key->next) {
+      DwOffset info_off = DW_INVALID_OFFSET;
+      SymsString test_name = syms_string_init(0,0);
+      switch (context->mode) {
+      case DW_MODE_32BIT: {
+        DwNameTableEntry32 *entry = (DwNameTableEntry32 *)key->entry;
+        info_off  = entry->info_off;
+        test_name = syms_string_init_lit(entry->name);
+      } break;
+      case DW_MODE_64BIT: {
+        DwNameTableEntry64 *entry = (DwNameTableEntry64 *)key->entry;
+        info_off  = entry->info_off;
+        test_name = syms_string_init_lit(entry->name);
+      } break;
+      case DW_MODE_NULL: break;
+      }
+      if (syms_string_cmp(name, test_name)) {
+        DwNameIter name_iter;
+        if (dw_name_iter_init(&name_iter, context, table_index)) {
+          /* NOTE(nick): Figuring out Comile Unit. */
+          DwOffset cu_info_base = DW_INVALID_OFFSET;
+          while (dw_name_iter_next_table(&name_iter)) {
+            U8 *start = DW_CAST(U8 *, name_iter.sec.data) + name_iter.unit_start;
+            U8 *end   = DW_CAST(U8 *, name_iter.sec.data) + name_iter.unit_end;
+            if (key->entry >= (void *)start && key->entry < (void *)end) {
+              cu_info_base = name_iter.cu_info_off;
+              break;
+            }
+          }
+          SYMS_ASSERT(cu_info_base != DW_INVALID_OFFSET);
+          if (tags_num < tags_max && cu_info_base != DW_INVALID_OFFSET) {
+            tags[tags_num++] = dw_tag_bake(context, cu_info_base, info_off);
+          }
+        }
+      }
+    }
+  }
+  return tags_num;
+}
+
+SYMS_INTERNAL syms_bool
+dw_name_table_init(DwContext *context, DwNameTableIndex table_index, SymsArena *arena)
+{
+  DwNameTable *table;
+  syms_bool is_inited = syms_false;
+
+  context->name_tables[table_index] = syms_arena_push_struct(arena, DwNameTable);
+  table = context->name_tables[table_index];
+  if (table) {
+    DwNameIter name_iter;
+    if (dw_name_iter_init(&name_iter, context, table_index)) {
+      DwTag tag;
+      SymsString name;
+
+      is_inited = syms_true;
+      while (dw_name_iter_next(&name_iter, &tag, &name)) {
+        DwNameTableKeyValue *key = syms_arena_push_struct(arena, DwNameTableKeyValue);
+        if (key) {
+          U32 key_index = syms_hash_djb2(name.data, name.len) % SYMS_ARRAY_SIZE(table->keys);
+          key->entry = name_iter.current_entry;
+          key->next = table->keys[key_index];
+          table->keys[key_index] = key;
+        } else {
+          is_inited = syms_false;
+          break;
+        }
+      }
+    }
+  }
+  return is_inited;
+}
+
+#if 0
+SYMS_INTERNAL void
+dw_name_table_free(DwContext *context, DwNameTableIndex table_index, void *ud)
+{
+  DwNameTable *table = &context->name_tables[table_index];
+  if (table->keys) {
+    U32 i;
+    for (i = 0; i < table->key_count; ++i) {
+      DwNameTableKeyValue *key;
+      for (key = table->keys[i]; key != 0; key = key->next) {
+        syms_free(key, ud);
+      }
+    }
+  }
+  syms_free(table->keys, ud);
+  table->key_count = 0;
+  table->keys = 0;
+}
+#endif
+
+SYMS_INTERNAL syms_bool
+dw_cat_init(DwAttribIter *iter, DwCommonAttribs *cat)
+{
+  DwAttribIter attrib_iter = *iter;
+  DwAttrib attrib;
+
+  DwAttrib hi_pc_attrib;
+  DwAttrib lo_pc_attrib;
+
+  //syms_bool is_address_internal = syms_true;
+
+  if (!dw_attrib_iter_reset(&attrib_iter))
+    return syms_false;
+
+  cat->decl_ln     = DW_INVALID_VALUE;
+  cat->decl_file   = DW_INVALID_VALUE;
+  cat->call_ln     = DW_INVALID_VALUE;
+  cat->call_file   = DW_INVALID_VALUE;
+  cat->len         = 0;
+  cat->rva         = DW_INVALID_OFFSET;
+  cat->range_off   = DW_INVALID_OFFSET;
+
+  cat->type_tag      = dw_invalid_ref();
+  cat->sibling_tag   = dw_invalid_ref();
+  cat->specification = dw_invalid_ref();
+
+  cat->linkage_name = syms_string_init(0,0);
+  cat->name         = syms_string_init(0,0);
+
+  hi_pc_attrib.value_class = DW_AT_CLASS_UNDEFINED;
+  lo_pc_attrib.value_class = DW_AT_CLASS_UNDEFINED;
+
+  while (dw_attrib_iter_next(&attrib_iter, &attrib)) {
+    switch (attrib.name) {
+    case DW_AT_SPECIFICATION: dw_attrib_get_ref(&attrib, &cat->specification);   break;
+    case DW_AT_LINKAGE_NAME:  dw_attrib_get_string(&attrib, &cat->linkage_name); break;
+    case DW_AT_NAME:          dw_attrib_get_string(&attrib, &cat->name);         break;
+    case DW_AT_DECL_FILE:     dw_attrib_get_const32(&attrib, &cat->decl_file);   break;
+    case DW_AT_DECL_LINE:     dw_attrib_get_const32(&attrib, &cat->decl_ln);     break;
+    case DW_AT_CALL_FILE:     dw_attrib_get_const32(&attrib, &cat->call_file);   break;
+    case DW_AT_CALL_LINE:     dw_attrib_get_const32(&attrib, &cat->call_ln);     break;
+    case DW_AT_SIBLING:       dw_attrib_get_ref(&attrib, &cat->sibling_tag);     break;
+    case DW_AT_TYPE:          dw_attrib_get_ref(&attrib, &cat->type_tag);        break;
+    case DW_AT_HIGH_PC:       hi_pc_attrib = attrib;                             break;
+    case DW_AT_LOW_PC:        lo_pc_attrib = attrib;                             break;
+    //case DW_AT_EXTERNAL:      is_address_internal = syms_false;                  break;
+    case DW_AT_RANGES:        dw_attrib_get_rnglist(&attrib, &cat->range_off);   break;
+    default: break;
+    }
+  }
+  if (syms_string_is_null(cat->name)) {
+    cat->name = cat->linkage_name;
+  }
+
+  if (dw_attrib_get_addr(&lo_pc_attrib, &cat->rva)) {
+    if (hi_pc_attrib.value_class == DW_AT_CLASS_CONST) {
+      if (!dw_attrib_get_const32(&hi_pc_attrib, &cat->len))
+        cat->rva = DW_INVALID_VALUE;
+    } else if (hi_pc_attrib.value_class == DW_AT_CLASS_ADDRESS) {
+      SymsAddr addr;
+      if (dw_attrib_get_addr(&hi_pc_attrib, &addr))
+        cat->len = syms_trunc_u32(addr - cat->rva);
+      else
+        cat->rva = DW_INVALID_VALUE;
+    }
+  }
+
+  return attrib_iter.is_exhausted;
+}
+
+SYMS_INTERNAL syms_bool
+dw_build_abbrev(DwCompileUnit *cu, SymsArena *arena)
+{
+  DwAbbrevIter abbrev_iter;
+  if (dw_abbrev_iter_init(&abbrev_iter, cu->dwarf, cu->abbrev_base)) {
+    DwAbbrevEntry abbrev;
+    // count abbrev tags
+    U32 tag_count = 0;
+    u32 tag_index = 0;
+    while (dw_abbrev_iter_next(&abbrev_iter, &abbrev)) {
+      if (abbrev.type == DW_ABBREV_ENTRY_TYPE_TAG_INFO) {
+        tag_count += 1;
+      } else if (abbrev.type == DW_ABBREV_ENTRY_TYPE_DIE_END) {
+        break;
+      }
+    }
+    // allocate abbrev lookup table
+    cu->abbrev_table.entry_count = tag_count;
+    cu->abbrev_table.entries = syms_arena_push_array(arena, DwAbbrevTableEntry, tag_count);
+    // load tag offsets
+    if (!dw_abbrev_iter_init(&abbrev_iter, cu->dwarf, cu->abbrev_base)) {
+      SYMS_INVALID_CODE_PATH;
+    }
+    while (dw_abbrev_iter_next(&abbrev_iter, &abbrev)) {
+      if (abbrev.type == DW_ABBREV_ENTRY_TYPE_TAG_INFO) {
+        cu->abbrev_table.entries[tag_index].id = syms_trunc_u32(abbrev.u.tag_info.id);
+        cu->abbrev_table.entries[tag_index].off = abbrev.data_lo;
+        tag_index += 1;
+      } else if (abbrev.type == DW_ABBREV_ENTRY_TYPE_DIE_END) {
+        break;
+      }
+    }
+    SYMS_ASSERT(tag_index == tag_count);
+  }
+  return syms_true;
+}
+
+/* TODO(nick): rename to dw_parse_compile_unit? */
+SYMS_INTERNAL syms_bool 
+dw_cu_init(DwCompileUnit *cu, DwContext *context, DwOffset info_cu_base)
+{
+  DwBinRead info = dw_bin_read_for_sec(context, DW_MODE_NULL, 0, DW_SEC_INFO);
+  DwAttribIter attrib_iter;
+  DwOffset addr_base, str_base, loclist_base, rnglist_base;
+  DwAttrib attrib;
+
+  cu->index = 0;
+  while (info.off < info_cu_base) {
+    DwOffset len = dw_bin_read_length(&info);
+    if (!dw_bin_skip(&info, len)) {
+        return syms_false;
+    }
+    cu->index += 1;
+  }
+  if (info.off == info.max) {
+    return syms_false;
+  }
+
+  u32 length = dw_bin_peek_u32(&info);
+  cu->mode = length < SYMS_UINT32_MAX ? DW_MODE_32BIT : DW_MODE_64BIT;
+  info.mode = cu->mode;
+
+  cu->dwarf = context;
+  cu->info_base = info.off;
+  cu->info_data_start = dw_bin_at(&info);
+  cu->producer = syms_string_init(0,0);
+  cu->compile_dir = syms_string_init(0,0);
+  cu->lang = DW_LANG_INVALID;
+  cu->case_type = DW_IDENTIFIER_CASE_SENSITIVE;
+  cu->use_utf8 = syms_false;
+  cu->line_base = DW_INVALID_VALUE; 
+  cu->dwo_name = syms_string_init(0,0);
+  cu->dwo_id = 0;
+  dw_seg_off_array_zero(&cu->addrs_arr);
+  dw_seg_off_array_zero(&cu->loclists_arr);
+  dw_seg_off_array_zero(&cu->rnglists_arr);
+  dw_off_array_zero(&cu->stroffs_arr);
+
+
+  /* NOTE(nick): The length that we are about to read does not include 
+   * bytes that indicate the length, so we have to add it here. */
+  cu->info_len = dw_bin_read_length(&info) + (info.off - cu->info_base);
+  cu->ver = dw_bin_read_u16(&info);
+
+  switch (cu->ver) {
+  case DWARF_V3:
+  case DWARF_V4: {
+    cu->unit_type   = DW_UT_RESERVED;
+    cu->abbrev_base = dw_bin_read_length(&info);
+    cu->addr_size   = dw_bin_read_u08(&info);
+  } break;
+  case DWARF_V5: {
+    cu->unit_type   = DW_CAST(DwUnitType, dw_bin_read_u08(&info));
+    cu->addr_size   = dw_bin_read_u08(&info);
+    cu->abbrev_base = dw_bin_read_length(&info);
+  } break;
+  default: return syms_false;
+  }
+  /* NOTE(nick): Have to update this, because at the moment of creation
+   * address size is not available. */
+  info.addr_size = cu->addr_size;
+
+  cu->attribs_off = info.off - cu->info_base;
+
+  /* NOTE(nick): Attributes come right after the header. */
+  if (!dw_attrib_iter_init2(&attrib_iter, cu, cu->attribs_off, cu->abbrev_base)) {
+    return syms_false;
+  }
+
+  if (attrib_iter.tag_type != DW_TAG_COMPILE_UNIT) {
+    return syms_false;
+  }
+
+  {
+    DwCommonAttribs cat;
+
+    if (!dw_cat_init(&attrib_iter, &cat))
+      return syms_false;
+
+    cu->name = cat.name;
+    cu->rva  = cat.rva;
+    cu->len  = cat.len;
+    cu->range_off = cat.range_off;
+  }
+
+  addr_base = str_base = loclist_base = rnglist_base = DW_INVALID_OFFSET;
+  while (dw_attrib_iter_next(&attrib_iter, &attrib)) {
+    U32 u;
+
+    switch (attrib.name) {
+    case DW_AT_PRODUCER:         dw_attrib_get_string(&attrib, &cu->producer);     break;
+    case DW_AT_NAME:             dw_attrib_get_string(&attrib, &cu->name);         break;
+    case DW_AT_COMP_DIR:         dw_attrib_get_string(&attrib, &cu->compile_dir);  break;
+    case DW_AT_STMT_LIST:        dw_attrib_get_lineptr(&attrib, &cu->line_base);   break;
+    case DW_AT_ADDR_BASE:        dw_attrib_get_addrptr(&attrib, &addr_base);       break;
+    case DW_AT_STR_OFFSETS_BASE: dw_attrib_get_stroffptr(&attrib, &str_base);      break;
+    case DW_AT_LOCLISTS_BASE:    dw_attrib_get_loclistptr(&attrib, &loclist_base); break;
+    case DW_AT_RNGLISTS_BASE:    dw_attrib_get_rnglistptr(&attrib, &rnglist_base); break;
+
+    case DW_AT_GNU_DWO_NAME:
+    case DW_AT_DWO_NAME: dw_attrib_get_string(&attrib, &cu->dwo_name);  break;
+
+    case DW_AT_GNU_DWO_ID:   dw_attrib_get_const(&attrib, &cu->dwo_id); break;
+
+    case DW_AT_USE_UTF8: {
+      U64 value;
+      if (dw_attrib_get_flag(&attrib, &value)) {
+        cu->use_utf8 = (syms_bool)value;
+      }
+    } break;
+    case DW_AT_LANGUAGE: {
+      if (dw_attrib_get_const32(&attrib, &u)) {
+        cu->lang = DW_CAST(DwLang, u);
+      }
+    } break;
+
+    case DW_AT_IDENTIFIER_CASE: {
+      if (dw_attrib_get_const32(&attrib, &u))
+        cu->case_type = DW_CAST(DwIdentifierCaseType, u);
+      break;
+    }
+    case DW_AT_MACRO_INFO:
+    case DW_AT_BASE_TYPES:
+    case DW_AT_MAIN_SUBPROGRAM:
+    default: break;
+    }
+  }
+
+  dw_seg_off_array_init(&cu->addrs_arr, context, cu->mode, DW_SEC_ADDR, addr_base);
+  dw_seg_off_array_init(&cu->loclists_arr, context, cu->mode, DW_SEC_LOCLISTS, loclist_base);
+  dw_seg_off_array_init(&cu->rnglists_arr, context, cu->mode, DW_SEC_RNGLISTS, rnglist_base);
+  dw_off_array_init(&cu->stroffs_arr, context, cu->mode, DW_SEC_STR_OFFSETS, str_base);
+
+  return syms_true;
+}
+
+SYMS_INTERNAL syms_bool
+dw_cu_iter_init(DwCuIter *iter, DwContext *context)
+{
+  iter->dwarf = context;
+  iter->err = syms_false;
+  iter->num_read = 0;
+  iter->next_cu = 0;
+  return syms_true;
+}
+
+SYMS_INTERNAL syms_bool
+dw_cu_iter_next(DwCuIter *iter, DwCompileUnit *cu)
+{
+  syms_bool result = syms_false;
+  if (dw_cu_init(cu, iter->dwarf, iter->next_cu)) {
+    ++iter->num_read;
+    iter->next_cu = cu->info_base + cu->info_len;
+    result = syms_true;
+  }
+  return result;
+}
+
+SYMS_INTERNAL void
+dw_line_iter_advance_pc(DwLineIter *line_iter, U64 advance)
+{
+  U64 op_index = line_iter->state.op_index + advance;
+  line_iter->state.address += line_iter->min_inst_len*(op_index/line_iter->max_ops_for_inst);
+  line_iter->state.op_index = op_index % line_iter->max_ops_for_inst;
+}
+
+SYMS_INTERNAL void
+dw_line_iter_reset_state(DwLineIter *iter)
+{
+  iter->state.address         = 0;
+  iter->state.op_index        = 0;
+  iter->state.file_index      = 1;
+  iter->state.line            = 1;
+  iter->state.column          = 0;
+  iter->state.is_stmt         = iter->default_is_stmt;
+  iter->state.basic_block     = syms_false;
+  iter->state.prologue_end    = syms_false;
+  iter->state.epilogue_begin  = syms_false;
+  iter->state.isa             = 0;
+  iter->state.discriminator   = 0;
+}
+
+DW_API syms_bool
+dw_line_iter_read_dir(DwBinRead *linesec, SymsString *dir_out)
+{
+  SymsString dir = dw_bin_read_string(linesec);
+  syms_bool result = syms_false;
+  if (!linesec->err) {
+    if (syms_string_peek_byte(dir, 0) != '\0') {
+      *dir_out = dir;
+      result = syms_true;
+    }
+  }
+  return result;
+}
+
+DW_API syms_bool
+dw_line_iter_read_file(DwBinRead *linesec, DwLineFile *file_out)
+{
+  syms_bool result = syms_false;
+  file_out->file_name = dw_bin_read_string(linesec);
+  if (syms_string_peek_byte(file_out->file_name, 0) != '\0') {
+    file_out->file_index  = 0;
+    file_out->dir_index   = dw_bin_read_uleb128(linesec);
+    file_out->file_index  = dw_bin_read_uleb128(linesec);
+    file_out->modify_time = dw_bin_read_uleb128(linesec);
+    file_out->file_size   = 0;
+    result = (linesec->err == syms_false);
+  }
+  return result;
+}
+
+SYMS_INTERNAL syms_bool
+dw_line_iter_get_dir(DwLineIter *iter, DwDirIndex index, SymsString *dir_out)
+{
+  syms_bool result = syms_false;
+
+  if (index == 0) {
+    *dir_out = iter->compile_dir;
+    result = syms_true;
+  } else if (index <= iter->dir_count) {
+    DwBinRead linesec = iter->dirs;
+    SymsString dir = syms_string_init(0,0);
+    DwDirIndex i;
+    for (i = 1; i <= index; ++i) {
+      if (!dw_line_iter_read_dir(&linesec, &dir)) {
+        break;
+      }
+    }
+    if (i >= index) {
+      *dir_out = dir;
+      result = syms_true;
+    }
+  }
+  return result;
+}
+
+SYMS_INTERNAL syms_bool
+dw_line_iter_get_file(DwLineIter *iter, DwFileIndex index, DwLineFile *file_out)
+{
+  syms_bool result = syms_false;
+  if (index == 0) {
+    syms_memset(file_out, 0, sizeof(*file_out));
+    file_out->file_name = iter->compile_file;
+    result = syms_true;
+  } else if (index <= iter->file_count) {
+    DwBinRead linesec = iter->files;
+    DwLineFile file;
+    DwFileIndex i;
+    
+    syms_memset(&file, 0, sizeof(file));
+    for (i = 1; i <= index; ++i) {
+      if (!dw_line_iter_read_file(&linesec, &file)) {
+        break;
+      }
+    }
+    if (i >= index) {
+      *file_out = file;
+      result = syms_true;
+    }
+  }
+  return result;
+}
+
+DW_API syms_bool
+dw_line_iter_init(DwLineIter *iter, DwCompileUnit *cu)
+{
+  syms_bool done = syms_false;
+  DwBinRead *linesec = &iter->linesec;
+
+  iter->compile_dir   = cu->compile_dir;
+  iter->compile_file  = cu->name;
+  iter->base_addr     = cu->rva;
+  iter->linesec       = dw_bin_read_for_sec(cu->dwarf, cu->mode, cu->addr_size, DW_SEC_LINE);
+
+  if (!dw_bin_seek(linesec, cu->line_base)) {
+    return syms_false;
+  }
+
+  iter->unit_length = dw_bin_read_length(linesec);
+  SYMS_ASSERT(linesec->max >= linesec->off + iter->unit_length);
+  linesec->max = linesec->off + iter->unit_length;
+
+  iter->ver = dw_bin_read_u16(linesec);
+  iter->header_len = dw_bin_read_offset(linesec);
+
+  if (iter->ver == DWARF_V5) {
+    if (cu->dwarf->secs[DW_SEC_INFO].data_len == 0) {
+      /* TODO(nick): DWARF5 might have only 2 sections available:
+       * ".debug_line" and ".debug_line_str". In this case
+       * parsing lines on a per CompileUnit basis is impossible 
+       * and something has to be done in order to make this work. */
+      SYMS_NOT_IMPLEMENTED;
+    }
+  }
+
+  iter->prog_off = linesec->off + iter->header_len;
+  iter->min_inst_len = dw_bin_read_u08(linesec);
+  switch (iter->ver) {
+  case DWARF_V5:
+  case DWARF_V4: {
+    iter->max_ops_for_inst = dw_bin_read_u08(linesec); 
+  } break;
+  case DWARF_V3:
+  case DWARF_V2: 
+  case DWARF_V1: {
+    iter->max_ops_for_inst = 1; 
+  } break;
+  default: SYMS_INVALID_CODE_PATH;
+  }
+  iter->default_is_stmt = dw_bin_read_u08(linesec);
+  iter->line_base       = dw_bin_read_s08(linesec);
+  iter->line_range      = dw_bin_read_u08(linesec);
+  iter->opcode_base     = dw_bin_read_u08(linesec);
+
+  if (iter->opcode_base == 0) {
+    return syms_false;
+  }
+  iter->num_opcode_lens = iter->opcode_base - 1u;
+  iter->opcode_lens = DW_CAST(U8 *, linesec->data) + linesec->off;
+  if (!dw_bin_skip(linesec, iter->num_opcode_lens*sizeof(U8))) {
+    return syms_false;
+  }
+
+  { /* NOTE(nick): Resolving memory range of the directory table */
+    DwOffset dir_lo = linesec->off;
+    DwOffset dir_hi;
+
+    SymsString dir;
+    while (dw_line_iter_read_dir(linesec, &dir)) {
+      iter->dir_count += 1;
+    }
+    dir_hi = linesec->off;
+    if (!dw_bin_subset(linesec, dir_lo, dir_hi - dir_lo, &iter->dirs)) {
+      SYMS_ASSERT_FAILURE("cannot subset dirs");
+      return syms_false;
+    }
+  }
+
+  { /* NOTE(nick): Resolving memory range of file-table */
+    DwOffset file_lo = linesec->off;
+    DwOffset file_hi;
+    DwLineFile file;
+    while (dw_line_iter_read_file(linesec, &file)) {
+      iter->file_count += 1;
+    }
+    file_hi = linesec->off;
+    if (!dw_bin_subset(linesec, file_lo, file_hi - file_lo, &iter->files)) {
+      SYMS_ASSERT_FAILURE("cannot subset files");
+      return syms_false;
+    }
+  }
+
+  if (!dw_bin_seek(linesec, iter->prog_off)) {
+    SYMS_ASSERT_FAILURE("offset for line-table program is invalid");
+    return syms_false;
+  }
+
+  dw_line_iter_reset_state(iter);
+
+  done = (linesec->err == syms_false);
+  return done;
+}
+
+SYMS_API syms_bool
+dw_line_iter_next(DwLineIter *iter, DwLineIterOp *op_out)
+{
+  DwBinRead *linesec = &iter->linesec;
+  syms_bool err = syms_false;
+  U8 opcode;
+
+  syms_bool line_emited = syms_false;
+  DwLine line;
+
+  syms_memset(&line, 0, sizeof(line));
+
+  next_opcode:;
+
+  SYMS_ASSERT(!line_emited);
+  if (linesec->off >= linesec->max)
+    return syms_false;
+
+  linesec = &iter->linesec;
+  opcode = dw_bin_read_u08(linesec);
+
+  if (opcode >= iter->opcode_base) {
+    /* NOTE(nick): Handling special opcode */
+    U32 adjusted_opcode = (U32)(opcode - iter->opcode_base);
+    U32 op_advance = adjusted_opcode / iter->line_range;
+    S32 line_inc = iter->line_base + DW_CAST(S32, adjusted_opcode) % iter->line_range;
+
+    iter->state.address += ((iter->state.op_index + op_advance)/iter->max_ops_for_inst)*iter->min_inst_len;
+    iter->state.op_index = (iter->state.op_index + op_advance) % iter->max_ops_for_inst;
+    iter->state.line = (DwLn)DW_CAST(U64, DW_CAST(S32, iter->state.line) + line_inc);
+    iter->state.basic_block = syms_false;
+    iter->state.epilogue_begin = syms_false;
+    iter->state.discriminator = 0;
+
+    line = iter->state;
+    line_emited = syms_true;
+
+    goto exit;
+  }
+
+  switch (opcode) { 
+  default: {
+    /* NOTE(nick): Skipping unknown opcode. This is a valid case and
+     * it works because compiler stores operand lengths that we can read 
+     * to skip unknown opcode */
+    U8 num_operands;
+    U8 i;
+
+    SYMS_ASSERT(opcode <= iter->num_opcode_lens);
+    num_operands = iter->opcode_lens[opcode - 1];
+    for (i = 0; i < num_operands; ++i)
+      dw_bin_read_uleb128(linesec);
+  } break;
+    /*
+     * NOTE(nick): Special opcodes
+     */
+  case DW_LNS_SPECIAL_OPCODE: {
+    U64 length = dw_bin_read_uleb128(linesec);
+    U64 start_off = linesec->off;
+    U64 num_skip;
+    opcode = dw_bin_read_u08(linesec);
+
+    switch (opcode) { 
+    case DW_LNE_END_SEQUENCE: {
+      iter->state.end_sequence = syms_true;
+
+      line = iter->state;
+      line_emited = syms_true;
+
+      dw_line_iter_reset_state(iter);
+    } break;
+
+    case DW_LNE_SET_ADDRESS: {
+      iter->state.address = dw_bin_read_addr(linesec);
+      iter->state.op_index = 0;
+    } break;
+
+    case DW_LNE_DEFINE_FILE: {
+      op_out->type = DW_LINE_ITER_OP_DEFINE_FILE;
+      dw_line_iter_read_file(linesec, &op_out->u.file);
+      return syms_true;
+    }
+
+    case DW_LNE_SET_DISCRIMINATOR: {
+      iter->state.discriminator = dw_bin_read_uleb128(linesec);
+    } break;
+
+    default: break;
+    }
+
+    num_skip = linesec->off - (start_off + length);
+    if (!dw_bin_skip(linesec, num_skip)) {
+      err = syms_true;
+      goto exit;
+    }
+
+    if (opcode == DW_LNE_END_SEQUENCE)
+      goto exit;
+  } break;
+
+    /*
+     * NOTE(nick): Standard opcodes
+     */
+  case DW_LNS_COPY: {
+    line = iter->state;
+    line_emited = syms_true;
+
+    iter->state.discriminator   = 0;
+    iter->state.basic_block     = syms_false;
+    iter->state.prologue_end    = syms_false;
+    iter->state.epilogue_begin  = syms_false;
+  } break;
+
+  case DW_LNS_ADVANCE_PC: {
+    U64 advance = dw_bin_read_uleb128(linesec);
+    dw_line_iter_advance_pc(iter, advance);
+  } break;
+
+  case DW_LNS_ADVANCE_LINE: {
+    S64 s = dw_bin_read_sleb128(linesec);
+    S64 l = DW_CAST(S64, iter->state.line);
+    S64 r = l + s;
+    SYMS_ASSERT(r >= 0);
+    iter->state.line = DW_CAST(DwLn, r);
+  } break;
+
+  case DW_LNS_SET_FILE: {
+    iter->state.file_index = dw_bin_read_uleb128(linesec); 
+  } break;
+
+  case DW_LNS_SET_COLUMN: {
+    iter->state.column = dw_bin_read_uleb128(linesec);
+  } break;
+
+  case DW_LNS_NEGATE_STMT: {
+    iter->state.is_stmt = !iter->state.is_stmt;
+  } break;
+
+  case DW_LNS_SET_BASIC_BLOCK: {
+    iter->state.basic_block = syms_true;
+  } break;
+
+  case DW_LNS_CONST_ADD_PC: {
+    U64 advance = (0xffu - iter->opcode_base) / iter->line_range;
+    dw_line_iter_advance_pc(iter, advance);
+  } break;
+
+  case DW_LNS_FIXED_ADVANCE_PC: {
+    U16 operand       = dw_bin_read_u16(linesec);
+    iter->state.address += operand;
+    iter->state.op_index = 0;
+  } break;
+
+  case DW_LNS_SET_PROLOGUE_END: {
+    iter->state.prologue_end = syms_true;
+  } break;
+
+  case DW_LNS_SET_EPILOGUE_BEGIN: {
+    iter->state.epilogue_begin = syms_true;
+  } break;
+
+  case DW_LNS_SET_ISA: {
+    iter->state.isa = dw_bin_read_uleb128(linesec);
+  } break;
+  }
+
+exit:;
+
+   if (iter->linesec.err)
+     err = syms_true;
+   if (err)
+     return syms_false;
+
+   if (line_emited) {
+     op_out->type = DW_LINE_ITER_OP_LINE;
+     op_out->u.line = line;
+   } else {
+     goto next_opcode;
+   }
+
+   return syms_true;
+}
+
+SYMS_API syms_bool
+dw_member_iter_init(DwMemberIter *iter, DwContext *context, DwTag udt_tag)
+{
+  if (!dw_attrib_iter_init_from_tag(&iter->attrib_iter, context, udt_tag)) {
+    return syms_false;
+  }
+  if (iter->attrib_iter.tag_type != DW_TAG_STRUCTURE_TYPE) {
+    return syms_false;
+  }
+  if (!iter->attrib_iter.has_children) {
+    return syms_false;
+  }
+  iter->depth = 0;
+  return syms_true;
+}
+
+SYMS_API syms_bool
+dw_member_iter_next(DwMemberIter *member_iter, DwMember *member)
+{
+  DwAttribIter *attrib_iter = &member_iter->attrib_iter;
+  syms_bool is_result_valid = syms_false;
+
+again:;
+  if (!dw_attrib_iter_next_tag(attrib_iter)) return syms_false;
+  if (attrib_iter->is_exhausted)             return syms_false;
+
+  switch (attrib_iter->tag_type) {
+  case DW_TAG_MEMBER: {
+    DwAttrib attrib;
+
+    member->name     = syms_string_init(0,0);
+    member->byte_off = DW_INVALID_OFFSET;
+    member->type_tag = dw_tag_bake_null();
+
+    for (;;) {
+      if (!dw_attrib_iter_next(attrib_iter, &attrib)) break;
+      switch (attrib.name) {
+      case DW_AT_NAME:                 dw_attrib_get_string(&attrib, &member->name);      break;
+      case DW_AT_DATA_MEMBER_LOCATION: dw_attrib_get_const(&attrib, &member->byte_off);   break;
+      case DW_AT_TYPE:                 dw_attrib_get_ref_tag(&attrib, attrib_iter->cu->dwarf, &member->type_tag); break;
+      default: break;
+      }
+    }
+
+    is_result_valid = member->byte_off != DW_INVALID_OFFSET && dw_tag_is_valid(member->type_tag);
+    break;
+  }
+  default:
+  case DW_TAG_UNION_TYPE:
+  case DW_TAG_STRUCTURE_TYPE: {
+    if (dw_attrib_iter_skip_children(attrib_iter)) goto again;
+    break;
+  }
+  }
+
+  return is_result_valid;
+}
+
+SYMS_API syms_bool
+dw_infer_type(DwContext *context, DwTag infer_type_tag, DwType *type)
+{
+  DwAttribIter attrib_iter;
+  DwAttrib attrib;
+  DwAttribTypeEncoding type_encoding;
+
+  if (!dw_attrib_iter_init_from_tag(&attrib_iter, context, infer_type_tag)) {
+    return syms_false;
+  }
+
+  type->kind          = DW_TYPE_NULL;
+  type->size          = 0;
+  type->name          = syms_string_init(0,0);
+  type->type_tag      = infer_type_tag;
+  type->next_type_tag = dw_tag_bake_null();
+
+  /* NOTE(nick): Parsing attributes that usually come with a type tag. */
+  type_encoding = DW_ATE_NULL;
+  while (dw_attrib_iter_next(&attrib_iter, &attrib)) {
+    switch (attrib.name) {
+    case DW_AT_ENCODING: {
+      if (attrib_iter.tag_type == DW_TAG_BASE_TYPE) {
+        U64 value;
+        if (!dw_attrib_get_const(&attrib, &value)) return syms_false;
+        type_encoding = DW_CAST(DwAttribTypeEncoding, value);
+      }
+      break;
+    }
+    case DW_AT_NAME: {
+      if (!dw_attrib_get_string(&attrib, &type->name)) return syms_false;
+      break;
+    }
+    case DW_AT_BYTE_SIZE: {
+      if (!dw_attrib_get_const(&attrib, &type->size)) return syms_false;
+      break;
+    }
+    case DW_AT_TYPE: {
+      dw_attrib_get_ref_tag(&attrib, context, &type->next_type_tag);
+      break;
+    }
+    default: break;
+    }
+  }
+
+  /* NOTE(nick): Collecting type modifiers. */
+  type->modifier = 0;
+  for (;;) {
+    switch (attrib_iter.tag_type) {
+    case DW_TAG_ATOMIC_TYPE:           type->modifier |= DW_TYPE_MDFR_ATOMIC;     break;
+    case DW_TAG_CONST_TYPE:            type->modifier |= DW_TYPE_MDFR_CONST;      break;
+    case DW_TAG_IMMUTABLE_TYPE:        type->modifier |= DW_TYPE_MDFR_IMMUTABLE;  break;
+    case DW_TAG_PACKED_TYPE:           type->modifier |= DW_TYPE_MDFR_PACKED;     break;
+    case DW_TAG_REFERENCE_TYPE:        type->modifier |= DW_TYPE_MDFR_REF;        break;
+    case DW_TAG_RESTRICT_TYPE:         type->modifier |= DW_TYPE_MDFR_RESTRICT;   break;
+    case DW_TAG_RVALUE_REFERENCE_TYPE: type->modifier |= DW_TYPE_MDFR_RVALUE_REF; break;
+    case DW_TAG_SHARED_TYPE:           type->modifier |= DW_TYPE_MDFR_SHARED;     break;
+    case DW_TAG_VOLATILE_TYPE:         type->modifier |= DW_TYPE_MDFR_VOLATILE;   break;
+    default: goto exit_modifier_read_loop;
+    }
+
+    if (!dw_attrib_iter_next_tag(&attrib_iter)) return syms_false;
+  }
+  exit_modifier_read_loop:;
+
+
+  /* NOTE(nick): Resolving type kind. */
+  switch (attrib_iter.tag_type) {
+  default:
+  case DW_TAG_NULL:           type->kind = DW_TYPE_NULL;    break;
+  case DW_TAG_CLASS_TYPE:     type->kind = DW_TYPE_CLASS;   break;
+  case DW_TAG_STRUCTURE_TYPE: type->kind = DW_TYPE_STRUCT;  break;
+  case DW_TAG_UNION_TYPE:     type->kind = DW_TYPE_UNION;   break;
+  case DW_TAG_TYPEDEF:        type->kind = DW_TYPE_TYPEDEF; break;
+  case DW_TAG_POINTER_TYPE:   type->kind = DW_TYPE_PTR;     break;
+  case DW_TAG_SUBROUTINE_TYPE: {
+    type->kind = DW_TYPE_PROC;
+    if (attrib_iter.has_children) {
+      if (!dw_attrib_iter_next_tag(&attrib_iter))           return syms_false;
+      if (attrib_iter.tag_type != DW_TAG_FORMAL_PARAMETER)  return syms_false;
+      SYMS_ASSERT(infer_type_tag.cu == type->next_type_tag.cu);
+      type->u.proc_params = dw_attrib_iter_get_tag(&attrib_iter);
+    } else {
+      type->u.proc_params = dw_tag_bake_null();
+    }
+  } break;
+  case DW_TAG_FORMAL_PARAMETER: {
+    if (attrib_iter.has_children) return syms_false;
+    if (!dw_attrib_iter_next_tag(&attrib_iter)) return syms_false;
+    type->kind = DW_TYPE_PROC_PARAM;
+    type->u.param_type = type->next_type_tag;
+    type->next_type_tag = dw_attrib_iter_get_tag(&attrib_iter);
+  } break;
+  case DW_TAG_ARRAY_TYPE: {
+    /* NOTE(nick): DW_TAG_ARRAY_TYPE must have children aka DW_TAG_SUBRANGE_TYPE,
+     * otherwise producer made an error. */
+    if (!attrib_iter.has_children) return syms_false;
+    if (!dw_attrib_iter_next_tag(&attrib_iter)) return syms_false;
+    if (attrib_iter.tag_type != DW_TAG_SUBRANGE_TYPE) return syms_false;
+    SYMS_ASSERT(infer_type_tag.cu == type->next_type_tag.cu);
+    infer_type_tag.abbrev = type->next_type_tag.info;
+  } /* Fall-through */
+  case DW_TAG_SUBRANGE_TYPE: {
+    type->kind = DW_TYPE_ARR;
+    type->u.arr_count = 0;
+
+    dw_attrib_iter_reset(&attrib_iter);
+    while (dw_attrib_iter_next(&attrib_iter, &attrib)) {
+      switch (attrib.name) {
+      case DW_AT_COUNT: {
+        if (!dw_attrib_get_const(&attrib, &type->u.arr_count)) return syms_false;
+        SYMS_ASSERT(type->u.arr_count != 0);
+        break;
+      }
+      /* TODO(nick): Would be nice to have these. */
+      case DW_AT_LOWER_BOUND:
+      case DW_AT_UPPER_BOUND:
+      case DW_AT_THREADS_SCALED:
+        SYMS_INVALID_CODE_PATH;
+      default: break;
+      }
+    }
+
+    if (dw_attrib_iter_next_tag(&attrib_iter)) {
+      if (attrib_iter.is_exhausted) {
+        SYMS_ASSERT(infer_type_tag.abbrev != infer_type_tag.info);
+        type->next_type_tag.info = infer_type_tag.abbrev;
+        type->next_type_tag.abbrev = 0;
+      } else {
+        if (attrib_iter.tag_type != DW_TAG_SUBRANGE_TYPE) return syms_false;
+        type->next_type_tag = dw_attrib_iter_get_tag(&attrib_iter);
+        type->next_type_tag.abbrev = infer_type_tag.abbrev;
+        SYMS_ASSERT(type->next_type_tag.abbrev != type->next_type_tag.info);
+      }
+    }
+    SYMS_ASSERT(infer_type_tag.abbrev);
+  } break;
+  case DW_TAG_BASE_TYPE: {
+    switch (type_encoding) {
+    case DW_ATE_SIGNED_CHAR:
+      type->modifier |= DW_TYPE_MDFR_CHAR;
+    /* NOTE(nick): Fall-through */
+    case DW_ATE_SIGNED: {
+      switch (type->size) {
+      case 1: type->kind = DW_TYPE_INT8;  break;
+      case 2: type->kind = DW_TYPE_INT16; break;
+      case 4: type->kind = DW_TYPE_INT32; break;
+      case 8: type->kind = DW_TYPE_INT64; break;
+      default: return syms_false;
+      }
+    } break;
+    case DW_ATE_UNSIGNED: {
+      switch (type->size) {
+      case 1: type->kind = DW_TYPE_UINT8;  break;
+      case 2: type->kind = DW_TYPE_UINT16; break;
+      case 4: type->kind = DW_TYPE_UINT32; break;
+      case 8: type->kind = DW_TYPE_UINT64; break;
+      default: return syms_false;
+      }
+    } break;
+    case DW_ATE_FLOAT: {
+      switch (type->size) {
+      case 4: type->kind = DW_TYPE_FLOAT32;  break;
+      case 8: type->kind = DW_TYPE_FLOAT64;  break;
+      default: return syms_false;
+      }
+    } break;
+    default: return syms_false;
+    }
+  } break;
+  }
+
+  return type->kind != DW_TYPE_NULL;
+}
+
+DW_API dw_uint
+dw_type_from_name(DwContext *context, const char *name, dw_uint name_len, dw_uint matches_max, DwTag *matches)
+{
+  SymsString name_str = syms_string_init(name, name_len);
+  dw_uint num_matched = dw_name_table_find(context, DW_NAME_TABLE_PUBTYPES, 
+                                           name_str,
+                                           matches_max, matches);
+  return num_matched;
+}
+
+DW_API syms_bool
+dw_global_from_name(DwContext *context, const char *name, dw_uint name_len, DwVar *var)
+{
+  DwTag var_tag;
+  syms_bool done = syms_false;
+  SymsString name_str = syms_string_init(name, name_len);
+  if (dw_name_table_find(context, DW_NAME_TABLE_PUBNAMES, name_str, 1, &var_tag)) {
+    DwAttribIter attrib_iter;
+    if (dw_attrib_iter_init_from_tag(&attrib_iter, context, var_tag))
+      done = dw_var_init(&attrib_iter, 0, 0, 0, var);
+  }
+
+  return done;
+}
+
+SYMS_INTERNAL syms_bool
+dw_proc_from_name(DwContext *context, char *name, dw_uint name_len, DwProc *proc)
+{
+  DwTag proc_tag;
+  syms_bool done = syms_false;
+  SymsString name_str = syms_string_init(name, name_len);
+  if (dw_name_table_find(context, DW_NAME_TABLE_PUBNAMES, name_str, 1, &proc_tag)) {
+    DwAttribIter attrib_iter;
+    if (dw_attrib_iter_init_from_tag(&attrib_iter, context, proc_tag))
+      done = dw_proc_init(&attrib_iter, proc);
+  }
+  return done;
+}
+
+SYMS_INTERNAL syms_bool
+dw_scope_init(DwAttribIter *attrib_iter, DwScope *scope)
+{
+  DwCommonAttribs cat;
+  syms_bool is_baked;
+
+  if (attrib_iter->tag_type != DW_TAG_LEXICAL_BLOCK) {
+    return syms_false;
+  }
+  if (!attrib_iter->has_children) {
+    return syms_false;
+  }
+  if (!dw_cat_init(attrib_iter, &cat)) {
+    return syms_false;
+  }
+
+  scope->name        = cat.name;
+  scope->decl_file   = cat.decl_file;
+  scope->decl_ln     = cat.decl_ln;
+  scope->rva         = cat.rva;
+  scope->len         = cat.len;
+  dw_ref_to_tag(attrib_iter->cu->dwarf, &cat.sibling_tag, &scope->sibling_tag);
+
+  is_baked = cat.rva != 0 && scope->len != DW_INVALID_VALUE;
+  return is_baked;
+}
+
+SYMS_INTERNAL syms_bool
+dw_var_init(DwAttribIter *attrib_iter, SymsAddr frame_base, SymsAddr member_location, SymsAddr cfa, DwVar *var)
+{
+  DwAttrib attrib;
+  syms_bool found_location;
+  syms_bool is_baked;
+
+  if (attrib_iter->tag_type != DW_TAG_VARIABLE) {
+    return syms_false;
+  }
+  if (!dw_attrib_iter_reset(attrib_iter)) {
+    return syms_false;
+  }
+
+  var->name      = syms_string_init(0,0);
+  var->flags     = 0;
+  var->type_tag  = dw_tag_bake_null();
+  var->decl_ln   = DW_INVALID_VALUE;
+  var->decl_file = DW_INVALID_VALUE;
+  syms_memset(&var->encoded_va, 0, sizeof(var->encoded_va));
+
+  if (attrib_iter->tag_type == DW_TAG_FORMAL_PARAMETER) {
+    var->flags |= DW_VAR_ARGUMENT;
+  }
+  else if (attrib_iter->tag_type != DW_TAG_VARIABLE) {
+    return syms_false;
+  }
+
+  found_location = syms_false;
+  while (dw_attrib_iter_next(attrib_iter, &attrib)) {
+    switch (attrib.name) {
+    case DW_AT_DECL_FILE: dw_attrib_get_const32(&attrib, &var->decl_ln);     break;
+    case DW_AT_DECL_LINE: dw_attrib_get_const32(&attrib, &var->decl_file);   break;
+    case DW_AT_NAME:      dw_attrib_get_string(&attrib, &var->name);         break;
+    case DW_AT_TYPE:      dw_attrib_get_ref_tag(&attrib, attrib_iter->cu->dwarf, &var->type_tag); break;
+
+    case DW_AT_LOCATION: {
+      var->encoded_va.cu = attrib_iter->cu;
+      var->encoded_va.u.expr.ops = attrib.value.exprloc.data;
+      var->encoded_va.u.expr.ops_size = dw_trunc_u32(attrib.value.exprloc.len);
+      var->encoded_va.u.expr.frame_base = frame_base;
+      var->encoded_va.u.expr.cfa = cfa;
+      var->encoded_va.u.expr.member_location  = member_location;
+
+      found_location = syms_true;
+      break;
+    }
+    default: break;
+    }
+  }
+
+  is_baked = !syms_string_is_null(var->name) && dw_tag_is_valid(var->type_tag) && found_location; 
+  return is_baked;
+}
+
+SYMS_INTERNAL syms_bool
+dw_class_init(DwAttribIter *attrib_iter, DwClass *udt)
+{
+  DwCommonAttribs cat;
+
+  if (attrib_iter->tag_type != DW_TAG_CLASS_TYPE) {
+    return syms_false;
+  }
+  if (!dw_attrib_iter_reset(attrib_iter)) {
+    return syms_false;
+  }
+  if (!dw_cat_init(attrib_iter, &cat)) {
+    return syms_false;
+  }
+
+  udt->name = cat.name;
+  udt->len = cat.len;
+
+  return syms_true;
+}
+
+SYMS_INTERNAL syms_bool
+dw_proc_init(DwAttribIter *attribs, DwProc *proc)
+{
+  DwCommonAttribs cat;
+  DwAttrib at;
+  syms_bool is_inited;
+
+  if (attribs->tag_type != DW_TAG_SUBPROGRAM && attribs->tag_type != DW_TAG_INLINED_SUBROUTINE)
+    return syms_false;
+
+  if (!dw_attrib_iter_reset(attribs))
+    return syms_false;
+
+  if (!dw_cat_init(attribs, &cat))
+    return syms_false;
+  
+  if (cat.rva == SYMS_ADDR_MAX)
+    return syms_false;
+
+  proc->name        = cat.name;
+  proc->decl_file   = cat.decl_file;
+  proc->decl_ln     = cat.decl_ln;
+  proc->call_file   = cat.call_file;
+  proc->call_ln     = cat.call_ln;
+  proc->len         = cat.len;
+  proc->type_tag    = dw_attrib_iter_get_tag(attribs);
+  proc->range_off   = cat.range_off;
+  proc->encoded_va  = dw_encode_null();
+
+  if (cat.len)
+    proc->encoded_va = dw_encode_rva(attribs->cu, cat.rva);
+
+  proc->frame_base.form = DW_FORM_INVALID;
+  proc->frame_base.value_class = DW_AT_CLASS_UNDEFINED;
+
+  while (dw_attrib_iter_next(attribs, &at)) {
+    switch (at.name) {
+    case DW_AT_FRAME_BASE:  proc->frame_base = at; break;
+
+    case DW_AT_ENTRY_PC: 
+    case DW_AT_START_SCOPE: {
+      /* TODO(nick): IMPLEMENET */
+    } break;
+                
+    case DW_AT_DECLARATION: {
+      return syms_false;
+    } break;
+
+    case DW_AT_SPECIFICATION: {
+      DwRef ref;
+      if (dw_attrib_get_ref(&at, &ref)) { 
+        DwCompileUnit *cu = attribs->cu;
+        DwAttribIter specs;
+        //
+        SYMS_ASSERT(ref.info >= cu->info_base && ref.info < cu->info_base + cu->info_len);
+        if (dw_attrib_iter_init(&specs, cu, ref.info - cu->info_base)) {
+          if (dw_cat_init(&specs, &cat)) {
+            proc->name = cat.name;
+            proc->decl_file = cat.decl_file;
+            proc->decl_ln = cat.decl_ln;
+          } else {
+            SYMS_ASSERT_FAILURE("Failed to parse attributes for DW_TAG_SUBPROGRAM");
+          }
+        } else {
+          SYMS_ASSERT_FAILURE("Failed to resolve compile unit from reference");
+        }
+      } else {
+        SYMS_ASSERT_FAILURE("Error on DW_TAG_SUBPROGRAM::DW_AT_SPECIFICATION");
+      }
+    } break;
+
+    case DW_AT_ABSTRACT_ORIGIN: {
+      DwRef ref;
+      if (dw_attrib_get_ref(&at, &ref)) {
+        DwCompileUnit *cu = attribs->cu;
+        DwAttribIter origin;
+        SYMS_ASSERT(ref.info > cu->info_base && ref.info < cu->info_base + cu->info_len);
+        //
+        if (dw_attrib_iter_init(&origin, cu, ref.info - cu->info_base)) {
+          if (dw_cat_init(&origin, &cat)) {
+            proc->name      = cat.name;
+            proc->decl_file = cat.decl_file;
+            proc->decl_ln   = cat.decl_ln;
+            //
+            if (cat.specification.info != DW_INVALID_OFFSET) {
+              DwAttribIter specs;
+              SYMS_ASSERT(cat.specification.info >= cu->info_base && cat.specification.info < cu->info_base + cu->info_len);
+              if (dw_attrib_iter_init(&specs, cu, cat.specification.info - cu->info_base)) {
+                if (dw_cat_init(&specs, &cat)) {
+                  proc->name = cat.name;
+                  proc->decl_file = cat.decl_file;
+                  proc->decl_ln = cat.decl_ln;
+                }
+              }
+            }
+          } else {
+            SYMS_ASSERT_FAILURE("Failed to parse attributes for DW_TAG_SUBPROGRAM");
+          }
+        } else {
+          SYMS_ASSERT_FAILURE("Failed to resolve compile unit from reference");
+        }
+      }
+    } break;
+
+    case DW_AT_EXTERNAL:
+    case DW_AT_ACCESSIBILITY:
+    case DW_AT_ADDRESS_CLASS:
+    case DW_AT_ARTIFICIAL:
+    case DW_AT_CALLING_CONVENTION:
+    case DW_AT_ELEMENTAL: 
+    case DW_AT_EXPLICIT:
+    case DW_AT_INLINE:
+    case DW_AT_LINKAGE_NAME:
+    case DW_AT_MAIN_SUBPROGRAM:
+    case DW_AT_OBJECT_POINTER:
+    case DW_AT_PROTOTYPED:
+    case DW_AT_PURE:
+    case DW_AT_RECURSIVE:
+    case DW_AT_RETURN_ADDR:
+    case DW_AT_SEGMENT:
+    case DW_AT_SIBLING:
+    case DW_AT_STATIC_LINK:
+    case DW_AT_TRAMPOLINE:
+    case DW_AT_TYPE:
+    case DW_AT_VISIBILITY:
+    case DW_AT_VIRTUALITY:
+    case DW_AT_VTABLE_ELEM_LOCATION: {
+      /* NOTE(nick): Ignore */
+      break;
+    }
+
+    default: break;
+    }
+  }
+  SYMS_ASSERT(attribs->is_exhausted);
+  
+  is_inited = ((dw_encoded_location_is_valid(proc->encoded_va) && proc->len > 0) ||proc->range_off != 0) && !syms_string_is_null(proc->name);
+  return is_inited;
+}
+
+SYMS_API syms_bool
+dw_proc_iter_init(DwProcIter *proc_iter, DwCompileUnit *cu)
+{
+  return dw_attrib_iter_init(&proc_iter->attribs, cu, cu->attribs_off);
+}
+
+SYMS_API syms_bool
+dw_proc_iter_next(DwProcIter *proc_iter, DwProc *proc)
+{
+  DwAttribIter *attribs = &proc_iter->attribs;
+  syms_bool is_result_valid = syms_false;
+  do {
+    if (attribs->tag_type == DW_TAG_SUBPROGRAM) {
+      // read subprogram tag info
+      is_result_valid = dw_proc_init(attribs, proc);
+      // fetch next subprogram tag, so next call to this iterator has new tag
+      dw_attrib_iter_next_tag(attribs);
+      // exit loop if subprogram was read
+      if (is_result_valid)
+        break;
+    }
+  } while (dw_attrib_iter_next_tag(attribs));
+  return is_result_valid;
+}
+
+SYMS_API syms_bool
+dw_local_iter_init(DwLocalIter *local_iter, DwContext *context, DwTag proc_tag)
+{
+  DwAttribIter *attrib_iter = &local_iter->attrib_iter;
+
+  if (!dw_attrib_iter_init_from_tag(attrib_iter, context, proc_tag)) {
+    return syms_false;
+  }
+  if (attrib_iter->tag_type != DW_TAG_SUBPROGRAM) {
+    return syms_false;
+  }
+  /* NOTE(nick): Children of this tag are variables and scopes for the specified procedure,
+   * and we have to make sure that they are there. */
+  if (!attrib_iter->has_children) {
+    return syms_false;
+  }
+
+  return syms_true;
+}
+
+SYMS_API syms_bool
+dw_local_iter_next(DwLocalIter *local_iter, DwLocal *local)
+{
+  DwAttribIter *attrib_iter = &local_iter->attrib_iter;
+
+  SymsAddr frame_base = 0;
+  SymsAddr member_location = 0;
+  SymsAddr cfa = 0;
+  syms_bool is_baked;
+
+  if (!dw_attrib_iter_next_tag(attrib_iter)) {
+    return syms_false;
+  }
+  if (attrib_iter->is_exhausted) {
+    return syms_false;
+  }
+
+  while (attrib_iter->tag_type == DW_TAG_LEXICAL_BLOCK) {
+    DwScope scope;
+    if (!dw_scope_init(attrib_iter, &scope)) {
+      return syms_false;
+    }
+    if (!dw_attrib_iter_skip_children(attrib_iter)) {
+      SYMS_ASSERT_PARANOID("unable to skip children");
+    }
+    break;
+  }
+
+  is_baked = dw_var_init(attrib_iter, frame_base, member_location, cfa, local);
+  return is_baked;
+}
+
+SYMS_API syms_bool
+dw_init(DwContext *context, SymsArch arch, DwInitdata *init_data)
+{
+  syms_bool is_inited = syms_false;
+
+  syms_memzero(context, sizeof(*context));
+  context->arch = arch;
+  context->mode = DW_MODE_32BIT;
+  syms_memcpy(context->secs, init_data->secs, sizeof(context->secs[0])*SYMS_ARRAY_SIZE(context->secs));
+  {
+    dw_uint i;
+    for (i = 0; i < SYMS_ARRAY_SIZE(context->secs); ++i) {
+      if (context->secs[i].data_len > SYMS_UINT32_MAX) {
+        context->mode = DW_MODE_64BIT;
+        break;
+      }
+    }
+  }
+  switch (context->mode) {
+  case DW_MODE_NULL:  context->msize_byte_count = 0; break;
+  case DW_MODE_32BIT: context->msize_byte_count = sizeof(U32); break;
+  case DW_MODE_64BIT: context->msize_byte_count = sizeof(U64); break;
+  }
+  is_inited = (context->secs[DW_SEC_INFO].data_len > 0 && context->secs[DW_SEC_ABBREV].data_len > 0);
+  return is_inited;
+}
+
+SYMS_API syms_bool
+dw_calc_heap_size(DwContext *dw, SymsUMM *size)
+{
+  DwNameIter name_iter;
+  DwTag tag;
+  SymsString name;
+
+  /* NOTE(nick): Counting public names. */
+  dw->pubnames_str_num = 0;
+  if (DW_HAS_SECTION(dw, DW_SEC_PUBNAMES)) {
+    if (dw_name_iter_init(&name_iter, dw, DW_NAME_TABLE_PUBNAMES)) {
+      while (dw_name_iter_next(&name_iter, &tag, &name)) {
+        dw->pubnames_str_num += 1;
+      }
+    }
+  }
+
+  /* NOTE(nick): Counting type names. */
+  dw->pubtypes_str_num = 0;
+  if (DW_HAS_SECTION(dw, DW_SEC_PUBTYPES)) {
+    if (dw_name_iter_init(&name_iter, dw, DW_NAME_TABLE_PUBTYPES)) {
+      while (dw_name_iter_next(&name_iter, &tag, &name)) {
+        dw->pubtypes_str_num += 1;
+      }
+    }
+  }
+
+  SYMS_ASSERT(dw->msize_byte_count == 4 || dw->msize_byte_count == 8);
+
+  *size = 0;
+  *size += SYMS_ARRAY_SIZE(dw->name_tables)*sizeof(*dw->name_tables[0]);
+  *size += (dw->pubtypes_str_num + dw->pubnames_str_num)*sizeof(DwNameTableKeyValue);
+
+  *size += 64; /* NOTE(nick): extra bytes for alignment */
+
+  return syms_true;
+}
+
+SYMS_API syms_bool
+dw_load_heap(DwContext *dw, SymsArena *arena)
+{
+  dw_name_table_init(dw, DW_NAME_TABLE_PUBNAMES, arena);
+  dw_name_table_init(dw, DW_NAME_TABLE_PUBTYPES, arena);
+  syms_bool is_result_valid = (~arena->flags & SYMS_ARENA_FLAG_ALLOC_FAILED);
+  return is_result_valid;
+}
+
+SYMS_API syms_bool
+dw_file_iter_init(DwFileIter *file_iter, DwContext *dwarf)
+{
+  DwCuIter *cu_iter = &file_iter->cu_iter;
+
+  dw_cu_iter_init(cu_iter, dwarf);
+  file_iter->next_cu = syms_true;
+
+  return syms_true;
+}
+
+SYMS_API dw_uint
+dw_file_iter_next(DwFileIter *file_iter, void *bf, dw_uint bf_max)
+{ 
+  DwCuIter *cu_iter = &file_iter->cu_iter;
+  DwCompileUnit *cu = &file_iter->cu;
+  DwLineIter *line_iter = &file_iter->line_iter;
+  dw_uint num_bytes_read = 0;
+
+  (void)bf; (void)bf_max;
+
+  if (file_iter->next_cu) {
+    if (!dw_cu_iter_next(cu_iter, cu)) return 0;
+    if (!dw_line_iter_init(line_iter, cu)) return 0;
+    file_iter->next_cu = syms_false;
+  }
+
+  return num_bytes_read;
+}
+
+SYMS_INTERNAL syms_bool
+dw_range_iter_init(DwRangeIter *iter, DwCompileUnit *cu, SymsOffset range_off)
+{
+  syms_bool is_inited = syms_false;
+  iter->rnglist = dw_bin_read_for_sec(cu->dwarf, cu->mode, cu->addr_size, DW_SEC_RANGES);
+  if (cu->ver == DWARF_V3 || cu->ver == DWARF_V4) {
+    if (dw_bin_seek(&iter->rnglist, range_off)) {
+      iter->base_addr = cu->rva;
+      if (cu->ver == DWARF_V3 || cu->ver == DWARF_V4) {
+        switch (cu->addr_size) {
+        case 4: {
+          U32 temp = dw_bin_peek_u32(&iter->rnglist);
+          if (temp == SYMS_UINT32_MAX) {
+            dw_bin_read_u32(&iter->rnglist);
+            iter->base_addr = dw_bin_read_u32(&iter->rnglist);
+          }
+          is_inited = syms_true;
+        } break;
+        case 8: {
+          U64 temp = dw_bin_peek_u64(&iter->rnglist);
+          if (temp == SYMS_UINT64_MAX) {
+            dw_bin_read_u64(&iter->rnglist);
+            iter->base_addr = dw_bin_read_u64(&iter->rnglist);
+          }
+          is_inited = syms_true;
+        } break;
+        }
+      }
+    }
+  } else {
+    // TODO(nick): DWARF5
+  }
+  return is_inited;
+}
+
+SYMS_INTERNAL syms_bool
+dw_range_iter_next(DwRangeIter *iter, SymsAddr *lo_out, SymsAddr *hi_out)
+{
+  syms_bool is_next_valid;
+  *lo_out = dw_bin_read_addr(&iter->rnglist);
+  *hi_out = dw_bin_read_addr(&iter->rnglist);
+  is_next_valid = *lo_out != 0 && *hi_out != 0 && !iter->rnglist.err;
+  *lo_out += iter->base_addr;
+  *hi_out += iter->base_addr;
+  return is_next_valid;
+}
+
+SYMS_INTERNAL syms_bool
+dw_range_check(DwCompileUnit *cu, SymsOffset range_off, SymsAddr addr, SymsAddr *lo_out, SymsAddr *hi_out)
+{
+  DwRangeIter range_iter;
+  syms_bool is_inrange = syms_false;
+  if (dw_range_iter_init(&range_iter, cu, range_off)) {
+    SymsAddr lo, hi;
+    while (dw_range_iter_next(&range_iter, &lo, &hi)) {
+      is_inrange = addr >= lo && addr < hi;
+      if (is_inrange) {
+        *lo_out = lo;
+        *hi_out = hi;
+        break;
+      }
+    }
+  }
+  return is_inrange;
+}
+
+SYMS_INTERNAL syms_bool
+dw_get_range_bounds(DwCompileUnit *cu, SymsOffset range_off, SymsAddr *lo_out, SymsAddr *hi_out)
+{
+  DwRangeIter range_iter;
+  syms_bool is_valid;
+  *lo_out = SYMS_ADDR_MAX;
+  *hi_out = 0;
+  if (dw_range_iter_init(&range_iter, cu, range_off)) {
+    SymsAddr lo, hi;
+    while (dw_range_iter_next(&range_iter, &lo, &hi)) {
+      *lo_out = SYMS_MIN(*lo_out, lo);
+      *hi_out = SYMS_MAX(*hi_out, hi);
+    }
+  }
+  is_valid = (*lo_out <= *hi_out);
+  return is_valid;
+}
+
+SYMS_API SymsString
+dw_tag_to_str(DwTagType tag)
+{
+  const char *tag_str;
+  SymsString result;
+  switch (tag) {
+#define X(name, value, version) case value: tag_str = #name; break;
+    DW_TAG_LIST
+#undef X
+    default: tag_str = ""; break;
+  }
+  result = syms_string_init_lit(tag_str);
+  return result;
+}
+
+SYMS_API SymsString
+dw_at_to_str(DwAttribType at)
+{
+  const char *at_str = "";
+  SymsString result;
+  switch (at) {
+#define X(name, value, ver, class_type) case value: at_str = #name; break;
+    DW_ATTRIB_LIST
+#undef X
+  default: at_str = "";
+  }
+  result = syms_string_init_lit(at_str);
+  return result;
+}
+
+SYMS_API SymsString
+dw_form_to_str(DwForm form)
+{
+  const char *form_str = "";
+  SymsString result;
+  switch (form) {
+#define X(name, value, ver, class_type) case value: form_str = #name; break;
+    DW_FORM_LIST
+#undef X
+  default: form_str = ""; break;
+  }
+  result = syms_string_init_lit(form_str);
+  return result;
+}
+
+SYMS_API SymsString
+dw_lang_to_str(DwLang lang)
+{
+  const char *lang_str = "";
+  SymsString result;
+  switch (lang) {
+  case DW_LANG_C89:         lang_str = "C89";            break;
+  case DW_LANG_C:           lang_str = "C";              break;
+  case DW_LANG_ADA83:       lang_str = "ADA83";          break;
+  case DW_LANG_C_PLUS_PLUS: lang_str = "C++";            break;
+  case DW_LANG_COBOL74:     lang_str = "COBOL74";        break;
+  case DW_LANG_COBOL85:     lang_str = "COBOL85";        break;
+  case DW_LANG_FORTAN77:    lang_str = "FORTAN77";       break;
+  case DW_LANG_FORTAN90:    lang_str = "FORTAN90";       break;
+  case DW_LANG_PASCAL83:    lang_str = "Pascal";         break;
+  case DW_LANG_MODULA2:     lang_str = "Modula-2";       break;
+  case DW_LANG_JAVA:        lang_str = "Java";           break;
+  case DW_LANG_C99:         lang_str = "C99";            break;
+  case DW_LANG_ADA95:       lang_str = "C95";            break;
+  case DW_LANG_FORTAN95:    lang_str = "FORTAN95";       break;
+  case DW_LANG_PLI:         lang_str = "PLI";            break;
+  case DW_LANG_OBJ_C:       lang_str = "Objective-C";    break;
+  case DW_LANG_OBJ_CPP:     lang_str = "Objective-C++";  break;
+  case DW_LANG_UPC:         lang_str = "UPC";            break;
+  case DW_LANG_D:           lang_str = "D";              break;
+  case DW_LANG_PYTHON:      lang_str = "Python";         break;
+  default:                  lang_str = "undefined";      break;
+  }
+  result = syms_string_init_lit(lang_str);
+  return result;
+}
+
+#if 0 
+SYMS_API syms_bool
+dw_aranges_unit_iter_init(DwContext *context, DwArangesUnitIter *iter)
+{
+  iter->aranges = dw_bin_read_for_sec(context, 0, DW_SEC_ARANGES);
+  return syms_true;
+}
+
+SYMS_API syms_bool
+dw_aranges_unit_iter_next(DwArangesUnitIter *iter, DwArangesUnit *unit_out)
+{
+  DwBinRead *aranges = &iter->aranges;
+  unit_out->unit_length = dw_bin_read_length(aranges);
+  DwOffset unit_offset = aranges->off;
+  unit_out->version = dw_bin_read_u16(aranges);
+  unit_out->debug_info_offset = dw_bin_read_offset(aranges);
+  unit_out->addr_size = dw_bin_read_u08(aranges);
+  unit_out->seg_size = dw_bin_read_u08(aranges);
+
+  // Round tuples offset to multiple of address size
+  unit_out->tuples_offset = aranges->off;
+  U32 range_entry_size = unit_out->addr_size * 2 + unit_out->seg_size;
+  DwOffset remainder = aranges->off % range_entry_size;
+  if (remainder != 0) {
+    unit_out->tuples_offset += (2 * unit_out->addr_size) - remainder;
+  }
+
+  unit_out->tuples_length = unit_out->unit_length - (unit_offset - unit_out->tuples_offset);
+  dw_bin_seek(aranges, unit_offset + unit_out->unit_length);
+  syms_bool is_valid = (aranges->err == syms_false);
+  return is_valid;
+}
+
+SYMS_API syms_bool
+dw_aranges_iter_init(DwContext *context, DwArangesUnit *unit, DwArangesIter *iter)
+{
+  DwBinRead aranges = dw_bin_read_for_sec(context, 0, DW_SEC_ARANGES);
+  syms_bool result = dw_bin_subset(&aranges, unit->tuples_offset, unit->tuples_length, &iter->aranges);
+  if (result) {
+    iter->aranges.addr_size = unit->addr_size;
+    iter->seg_size = unit->seg_size;
+  }
+  return result;
+}
+
+SYMS_API syms_bool
+dw_aranges_iter_next(DwArangesIter *iter, SymsAddr *addr_out, U64 *len_out)
+{
+  unit_out->tuples_offset = aranges->off;
+  U32 range_entry_size = unit_out->addr_size * 2 + unit_out->seg_size;
+  DwOffset remainder = aranges->off % range_entry_size;
+  if (remainder != 0) {
+    unit_out->tuples_offset += (2 * unit_out->addr_size) - remainder;
+  }
+
+  if (iter->seg_size > 0) {
+    dw_bin_read_u32(&iter->aranges);
+  }
+  *addr_out = dw_bin_read_addr(&iter->aranges);
+  *len_out = dw_bin_read_length(&iter->aranges);
+  syms_bool is_valid = (iter->aranges.err == syms_false);
+  return is_valid;
+}
+
+#endif
+
+#if 0
+
+typedef struct DwAttribNode
+{
+  DwAttrib data;
+  struct DwAttribNode *next;
+} DwAttribNode;
+
+typedef struct DwDieNode {
+  /* tag id */
+  U64 tag;
+
+  DwOffset info_off;
+
+  /* list of attributes */
+  DwAttribNode *attribs;
+
+  /* sub-children */
+  struct DwDieNode *sub;
+
+  /* next node */
+  struct DwDieNode *next;
+} DwDieNode;
+
+typedef struct DwDie {
+  DwDieNode *root;
+
+  DwOffset abbrev_base;
+
+  U32 offset_count;
+  DwOffset *offsets;
+
+  struct DwDie *next;
+} DwDie;
+
+typedef struct DwDieStack {
+  DwDieNode *node_head;
+  DwDieNode *node_last;
+  DwDieNode **next_node;
+  DwAttribNode **next_attrib;
+} DwDieStack;
+SYMS_API DwDie *
+
+SYMS_INTERNAL char *
+dw_stringify_tag(U32 tag)
+{
+#define X(name, value, verison) case value: return #name;
+  switch (tag) {
+    DW_TAG_LIST
+  }
+#undef X
+  return "";
+}
+
+SYMS_INTERNAL char *
+dw_stringify_attrib(U32 attrib)
+{
+#define X(name, value, version, class_flags) case value: return #name;
+  switch (attrib) {
+    DW_ATTRIB_LIST
+  }
+#undef X
+  return "";
+}
+
+SYMS_INTERNAL char *
+dw_stringify_form(U32 form)
+{
+  switch (form) {
+#define X(name, value, version, class_flags) case value: return #name;
+    DW_FORM_LIST
+#undef X
+  }
+  return "";
+}
+
+SYMS_INTERNAL void
+dw_dump_die(DwDie *die_list, void *userdata)
+{
+  U32 stack_max = 256;
+  DwDieStack *stack = (DwDieStack *)syms_malloc(sizeof(DwDieStack)*stack_max, userdata);
+  DwDie *die;
+
+  for (die = die_list; die != 0; die = die->next) {
+    U32 stack_index = 0;
+    DwAttribNode *attrib;
+    stack[stack_index].node_head = die->root;
+    do {
+      DwDieNode *node = stack[stack_index].node_head;
+      if (!stack[stack_index].node_head) {
+        stack_index -= 1;
+        continue;
+      }
+      stack[stack_index].node_head = node->next;
+      if (node->sub) {
+        stack_index += 1;
+        SYMS_ASSERT(stack_index < stack_max);
+        stack[stack_index].node_head = node->sub;
+      }
+      printf("<%u><%llx> %s\n", stack_index, node->info_off, dw_stringify_tag(node->tag));
+      for (attrib = node->attribs; attrib != 0; attrib = attrib->next) {
+        printf("\t%s %s\n", dw_stringify_attrib(attrib->data.name), dw_stringify_form(attrib->data.form));
+      }
+    } while (stack_index > 0);
+  }
+  syms_free(stack, userdata);
+}
+
+SYMS_API DwDie *
+dw_load_die(DwContext *dw, void *userdata)
+{
+  DwDie *die_list = 0;
+  DwDie **next_die = &die_list;
+  {
+    DwAbbrevIter abbrev_iter;
+    DwAbbrevEntry abbrev_entry;
+    if (dw_abbrev_iter_init(&abbrev_iter, dw, 0)) {
+      DwDie *die = 0;
+
+      /* build die list */
+      while (dw_abbrev_iter_next(&abbrev_iter, &abbrev_entry)) {
+        switch (abbrev_entry.type) {
+        case DW_ABBREV_ENTRY_TYPE_TAG_INFO: {
+          die->offset_count = SYMS_MAX(die->offset_count, abbrev_entry.u.tag_info.id + 1);
+        } break;
+
+        case DW_ABBREV_ENTRY_TYPE_ATTRIB_INFO: {
+        } break;
+        case DW_ABBREV_ENTRY_TYPE_ATTRIB_INFO_NULL: {
+        } break;
+
+        case DW_ABBREV_ENTRY_TYPE_DIE_BEGIN: {
+          die = (DwDie *)syms_malloc(sizeof(DwDie), userdata);
+          die->root = 0;
+          die->abbrev_base = 0;
+          die->offset_count = 0;
+          die->offsets = 0;
+          die->next = 0;
+          *next_die = die;
+          next_die = &die->next;
+        } break;
+        case DW_ABBREV_ENTRY_TYPE_DIE_END: {
+          die = 0;
+        } break;
+        }
+      }
+
+      if (!dw_abbrev_iter_init(&abbrev_iter, dw, 0)) {
+        SYMS_INVALID_CODE_PATH;
+      }
+      die = die_list;
+      while (dw_abbrev_iter_next(&abbrev_iter, &abbrev_entry)) {
+        switch (abbrev_entry.type) {
+        case DW_ABBREV_ENTRY_TYPE_DIE_BEGIN: {
+          die->abbrev_base = abbrev_entry.data_lo;
+          die->offsets = (DwOffset *)syms_malloc(sizeof(DwOffset)*die->offset_count, userdata);
+          x_assert(die->offsets);
+        } break;
+        case DW_ABBREV_ENTRY_TYPE_DIE_END: {
+          die = die->next;
+        } break;
+        case DW_ABBREV_ENTRY_TYPE_TAG_INFO: {
+          x_assert(abbrev_entry.u.tag_info.id < die->offset_count);
+          die->offsets[abbrev_entry.u.tag_info.id] = abbrev_entry.data_lo;
+        } break;
+        }
+      }
+    }
+  }
+
+  {
+    /* create stack for traversing tree */
+    U32 stack_max = 256;
+    U32 stack_index = 0;
+    DwDieStack *stack = (DwDieStack *)syms_malloc(sizeof(DwDieStack)*stack_max, userdata);
+
+    /* init .debug_info reader */
+    DwBinRead info = dw_bin_read_for_sec(dw, 0, DW_SEC_INFO);
+
+    while (info.off < info.max) {
+      /* parsing .debug_info entry header */
+      U64 info_size = dw_bin_read_length(&info);
+      U16 ver = dw_bin_read_u16(&info);
+      U8 unit_type = 0;
+      DwOffset abbrev_base;
+      U8 addr_size;
+      if (ver == DWARF_V5) {
+        unit_type = dw_bin_read_u08(&info);
+        addr_size = dw_bin_read_u08(&info);
+        abbrev_base = dw_bin_read_length(&info);
+      } else {
+        abbrev_base = dw_bin_read_length(&info);
+        addr_size = dw_bin_read_u08(&info);
+      }
+      info.addr_size = addr_size;
+      /* find die from .debug_info header */
+      DwDie *die = 0;
+      for (die = die_list; die != 0; die = die->next) {
+        if (die->abbrev_base == abbrev_base) {
+          break;
+        }
+      }
+      if (!die) {
+        SYMS_ASSERT_FAILURE("cannot find corresponding DIE entry");
+      }
+      stack_index = 0;
+      stack[stack_index].node_head = 0;
+      stack[stack_index].next_node = &stack[stack_index].node_head;
+      stack[stack_index].next_attrib = 0;
+      do {
+        U64 abbrev_id = dw_bin_read_uleb128(&info);
+        if (abbrev_id > 0) {
+          if (abbrev_id >= die->offset_count) {
+            SYMS_ASSERT_FAILURE("found an invalid abbrev id");
+          }
+          DwOffset abbrev_offset = die->offsets[abbrev_id];
+          DwAbbrevIter abbrev_iter;
+          DwAbbrevEntry abbrev_entry;
+          if (!dw_abbrev_iter_init(&abbrev_iter, dw, abbrev_offset)) {
+            SYMS_ASSERT_FAILURE("invalid abbrev offset");
+          }
+
+          syms_bool has_children = syms_false;
+          while (dw_abbrev_iter_next(&abbrev_iter, &abbrev_entry)) {
+            switch (abbrev_entry.type) {
+            case DW_ABBREV_ENTRY_TYPE_DIE_BEGIN: break;
+
+            case DW_ABBREV_ENTRY_TYPE_ATTRIB_INFO_NULL:
+            case DW_ABBREV_ENTRY_TYPE_DIE_END: 
+              goto exit_loop;
+
+            case DW_ABBREV_ENTRY_TYPE_TAG_INFO: {
+              DwDieNode *node = (DwDieNode *)syms_malloc(sizeof(DwDieNode), userdata);
+              node->tag = abbrev_entry.u.tag_info.tag;
+              node->info_off = info.off;
+              node->attribs = 0;
+              node->sub = 0;
+              node->next = 0;
+
+              stack[stack_index].node_last = node;
+              *stack[stack_index].next_node = node;
+              stack[stack_index].next_node = &node->next;
+              stack[stack_index].next_attrib = &node->attribs;
+
+              has_children = abbrev_entry.u.tag_info.has_children;
+            } break;
+            case DW_ABBREV_ENTRY_TYPE_ATTRIB_INFO: {
+              DwAttribNode *attrib_node = (DwAttribNode *)syms_malloc(sizeof(DwAttribNode), userdata);
+              attrib_node->next = 0;
+              *stack[stack_index].next_attrib = attrib_node;
+              stack[stack_index].next_attrib = &attrib_node->next;
+
+              attrib_node->data.tag = dw_tag_bake_null();
+              attrib_node->data.name = (DwAttribType)abbrev_entry.u.attrib_info.name;
+              attrib_node->data.form = (DwForm)abbrev_entry.u.attrib_info.form;
+              attrib_node->data.value_class = DW_AT_CLASS_INVALID;
+
+              switch (abbrev_entry.u.attrib_info.form) {
+              case DW_FORM_BLOCK1:        attrib_node->data.form_value.block      = dw_bin_read_u08(&info);       break;
+              case DW_FORM_BLOCK2:        attrib_node->data.form_value.block      = dw_bin_read_u16(&info);       break;
+              case DW_FORM_BLOCK4:        attrib_node->data.form_value.block      = dw_bin_read_u32(&info);       break;
+              case DW_FORM_BLOCK:         attrib_node->data.form_value.block      = dw_bin_read_uleb128(&info);   break;
+              case DW_FORM_DATA1:         attrib_node->data.form_value.data       = dw_bin_read_u08(&info);       break;
+              case DW_FORM_DATA2:         attrib_node->data.form_value.data       = dw_bin_read_u16(&info);       break;
+              case DW_FORM_DATA4:         attrib_node->data.form_value.data       = dw_bin_read_u32(&info);       break;
+              case DW_FORM_DATA8:         attrib_node->data.form_value.data       = dw_bin_read_u64(&info);       break;
+              case DW_FORM_UDATA:         attrib_node->data.form_value.udata      = dw_bin_read_uleb128(&info);   break;
+              case DW_FORM_SDATA:         attrib_node->data.form_value.sdata      = dw_bin_read_sleb128(&info);   break;
+              case DW_FORM_REF1:          attrib_node->data.form_value.ref        = dw_bin_read_u08(&info);       break;
+              case DW_FORM_REF2:          attrib_node->data.form_value.ref        = dw_bin_read_u16(&info);       break;
+              case DW_FORM_REF4:          attrib_node->data.form_value.ref        = dw_bin_read_u32(&info);       break;
+              case DW_FORM_REF8:          attrib_node->data.form_value.ref        = dw_bin_read_u64(&info);       break;
+              case DW_FORM_REF_SIG8:      attrib_node->data.form_value.ref        = dw_bin_read_u64(&info);       break;
+              case DW_FORM_REF_ADDR:      attrib_node->data.form_value.ref        = dw_bin_read_offset(&info);    break;
+              case DW_FORM_REF_UDATA:     attrib_node->data.form_value.ref        = dw_bin_read_uleb128(&info);   break;
+              case DW_FORM_REF_SUP4:      attrib_node->data.form_value.ref        = dw_bin_read_u32(&info);       break;
+              case DW_FORM_REF_SUP8:      attrib_node->data.form_value.ref        = dw_bin_read_u64(&info);       break;
+              case DW_FORM_SEC_OFFSET:    attrib_node->data.form_value.sec_offset = dw_bin_read_offset(&info);    break;
+              case DW_FORM_ADDR:          attrib_node->data.form_value.addr       = dw_bin_read_addr(&info);      break;
+              case DW_FORM_FLAG:          attrib_node->data.form_value.flag       = dw_bin_read_u08(&info);       break;
+              case DW_FORM_FLAG_PRESENT:  attrib_node->data.form_value.flag       = 1;                            break;
+              case DW_FORM_STRP:          attrib_node->data.form_value.strp       = dw_bin_read_offset(&info);    break;
+              case DW_FORM_STRP_SUP:      attrib_node->data.form_value.strp       = dw_bin_read_offset(&info);    break;
+              case DW_FORM_LINE_STRP:     attrib_node->data.form_value.strp       = dw_bin_read_offset(&info);    break;
+              case DW_FORM_STRING:        attrib_node->data.form_value.string     = dw_bin_read_string(&info);    break;
+              case DW_FORM_STRX:          attrib_node->data.form_value.strx       = dw_bin_read_uleb128(&info);   break;
+              case DW_FORM_STRX1:         attrib_node->data.form_value.strx       = dw_bin_read_u08(&info);       break;
+              case DW_FORM_STRX2:         attrib_node->data.form_value.strx       = dw_bin_read_u16(&info);       break;
+              case DW_FORM_STRX3:         attrib_node->data.form_value.strx       = dw_bin_read_u24(&info);       break;
+              case DW_FORM_STRX4:         attrib_node->data.form_value.strx       = dw_bin_read_u32(&info);       break;
+              case DW_FORM_ADDRX:         attrib_node->data.form_value.addrx      = dw_bin_read_uleb128(&info);   break;
+              case DW_FORM_ADDRX1:        attrib_node->data.form_value.addrx      = dw_bin_read_u08(&info);       break;
+              case DW_FORM_ADDRX2:        attrib_node->data.form_value.addrx      = dw_bin_read_u16(&info);       break;
+              case DW_FORM_ADDRX3:        attrib_node->data.form_value.addrx      = dw_bin_read_u24(&info);       break;
+              case DW_FORM_ADDRX4:        attrib_node->data.form_value.addrx      = dw_bin_read_u32(&info);       break;
+              case DW_FORM_LOCLISTX:      attrib_node->data.form_value.loclistx   = dw_bin_read_uleb128(&info);   break;
+              case DW_FORM_RNGLISTX:      attrib_node->data.form_value.rnglistx   = dw_bin_read_uleb128(&info);   break;
+              case DW_FORM_DATA16: {
+                attrib_node->data.form_value.data16.lo = dw_bin_read_u64(&info);
+                attrib_node->data.form_value.data16.hi = dw_bin_read_u64(&info);
+              } break;
+              case DW_FORM_IMPLICIT_CONST: {
+                /*
+                 * NOTE(nick): This is special case. Unlike
+                 * other forms that have their values stored in
+                 * the .debug_info section, this one defines
+                 * it's value in the .debug_abbrev section. *
+                 */
+                attrib_node->data.form_value.cnst = abbrev_entry.u.attrib_info.implicit_const;
+              } break;
+              case DW_FORM_EXPRLOC: {
+                attrib_node->data.form_value.exprloc.len = dw_bin_read_uleb128(&info);
+                attrib_node->data.form_value.exprloc.data = DW_CAST(U8 *, &info.data) + info.off;
+                if (!dw_bin_skip(&info, attrib_node->data.form_value.exprloc.len)) {
+                  return syms_false;
+                }
+              } break;
+              }
+            } break;
+            }
+          }
+          exit_loop:;
+          if (has_children) {
+            stack_index += 1;
+            x_assert(stack_index < stack_max);
+            stack[stack_index].node_head = 0;
+            stack[stack_index].next_node = &stack[stack_index].node_head;
+            stack[stack_index].next_attrib = 0;
+          }
+        } else {
+          SYMS_ASSERT(stack_index > 0);
+          SYMS_ASSERT(!stack[stack_index - 1].node_last->sub);
+          SYMS_ASSERT(stack[stack_index].node_last);
+          stack[stack_index - 1].node_last->sub = stack[stack_index].node_head;
+          stack_index -= 1;
+        }
+      } while (stack_index > 0);
+      die->root = stack[0].node_head;
+    }
+
+    syms_free(stack, userdata);
+    stack = 0;
+  }
+
+#if 1
+  dw_dump_die(die_list, userdata);
+#endif
+
+  return 0;
+}
+#endif
