RAD symbol library

diff --git a/Engine/Source/Runtime/SymsLib/src/dwarf/syms_dwarf_api.c b/Engine/Source/Runtime/SymsLib/src/dwarf/syms_dwarf_api.c
new file mode 100644
--- /dev/null
+++ b/Engine/Source/Runtime/SymsLib/src/dwarf/syms_dwarf_api.c
@@ -0,0 +1,994 @@
+// Copyright Epic Games, Inc. All Rights Reserved.
+/******************************************************************************
+ * File   : syms_dwarf.c                                                      *
+ * Author : Nikita Smith                                                      *
+ * Created: 2020/06/20                                                        *
+ * Purpose: API wrappers and type conversion for dwarf.h                      *
+ ******************************************************************************/
+
+SYMS_INTERNAL const DwTag *
+syms_typeid_to_dw(const SymsTypeID *type_id)
+{
+  SYMS_ASSERT(type_id->kind == SYMS_TYPE_ID_DW);
+  return (const DwTag *)type_id->impl;
+}
+
+SYMS_INTERNAL SymsTypeID
+syms_typeid_for_dw(DwTag tag)
+{
+  SymsTypeID type_id;
+  type_id.kind = SYMS_TYPE_ID_DW;
+  *(DwTag *)type_id.impl = tag;
+  return type_id;
+}
+
+SYMS_INTERNAL syms_int
+syms_typeid_cmp_dw(DwTag a, DwTag b)
+{
+  syms_int r = 0;
+  if (a.cu <= b.cu) {
+    if (a.info < b.info) {
+      r = -1;
+    } else if (a.info > b.info) {
+      r = +1;
+    }
+  } else if (a.cu > b.cu) {
+    r = +1;
+  }
+  return r;
+}
+
+// ----------------------------------------------------------------------------
+
+SYMS_INTERNAL syms_bool
+syms_mod_iter_init_dwarf(SymsDebugInfoDwarf *debug_info, SymsModIterDwarf *iter)
+{
+  if (dw_cu_iter_init(&iter->impl, &debug_info->context)) {
+    return syms_true;
+  } else {
+    syms_memzero(&iter, sizeof(iter));
+    return syms_false;
+  }
+}
+
+SYMS_INTERNAL syms_bool
+syms_mod_iter_next_dwarf(SymsDebugInfoDwarf *debug_info, SymsModIterDwarf *iter, SymsMod *mod_out)
+{ 
+  DwCompileUnit *cu = (DwCompileUnit *)mod_out->impl;
+  syms_bool result = syms_false;
+
+  (void)debug_info;
+
+  if (dw_cu_iter_next(&iter->impl, cu)) {
+    mod_out->id = (SymsModID)cu->info_base;
+    mod_out->name = syms_string_ref_str(cu->name);
+    mod_out->va = cu->rva;
+    mod_out->size = cu->len;
+    if (cu->len == 0) {
+      DwRangeIter range_iter;
+      if (dw_range_iter_init(&range_iter, cu, cu->range_off)) {
+        SymsAddr lo, hi;
+        while (dw_range_iter_next(&range_iter, &lo, &hi)) {
+          SYMS_ASSERT(lo <= hi);
+          mod_out->size += (hi - lo);
+        }
+      }
+    }
+    result = syms_true;
+  }
+  return result;
+}
+
+// ----------------------------------------------------------------------------
+
+SYMS_INTERNAL SymsLineIterDwarf *
+syms_line_iter_to_dw(SymsLineIter *iter)
+{
+  SymsLineIterDwarf *impl = (SymsLineIterDwarf *)iter->impl;
+  return impl;
+}
+
+SYMS_INTERNAL syms_bool
+syms_line_iter_init_dwarf(SymsDebugInfoDwarf *debug_info, SymsLineIterDwarf *iter, DwCompileUnit *cu, syms_bool *has_line_count)
+{ (void)debug_info;
+  *has_line_count = syms_false;
+  iter->prev_file_index = 0;
+  iter->cu = cu;
+  return dw_line_iter_init(&iter->impl, cu);
+}
+
+SYMS_INTERNAL syms_bool
+syms_line_iter_next_dwarf(SymsDebugInfoDwarf *debug_info, SymsLineIterDwarf *iter, SymsLine *line_out, syms_bool *switched_file, SymsSourceFile *file_out)
+{
+  syms_bool keep_fetching = syms_true;
+  syms_bool result = syms_false;
+(void)debug_info;
+
+  for (; keep_fetching ;) {
+    DwLineIter  *impl = &iter->impl;
+    DwLineIterOp dw_op;
+
+    keep_fetching = syms_false;
+
+    if (dw_line_iter_next(impl, &dw_op)) {
+      switch (dw_op.type) {
+      case DW_LINE_ITER_OP_NULL: break;
+      case DW_LINE_ITER_OP_LINE: {
+        dw_uint ln;
+        u16 col;
+
+        if (dw_op.u.line.file_index != iter->prev_file_index) {
+          SymsString dir;
+          DwLineFile file; 
+          dir = syms_string_init(0,0);
+          file.file_name = syms_string_init(0,0);
+          if (dw_line_iter_get_file(impl, dw_op.u.line.file_index, &file)) {
+            if (dw_line_iter_get_dir(impl, file.dir_index, &dir)) {
+              file_out->name = syms_string_ref_dw_path(iter->cu->compile_dir, dir, file.file_name);
+              file_out->chksum_type = SYMS_CHECKSUM_NULL; // DWARF doesn't export checksums
+            } else {
+              SYMS_ASSERT_FAILURE("cannot extract file name for line");
+            }
+            iter->prev_file_index = dw_op.u.line.file_index;
+          } else {
+            SYMS_ASSERT_FAILURE("cannot extract directory for line");
+          }
+
+          *switched_file = syms_true;
+        }
+        ln = dw_op.u.line.line;
+        col = syms_trunc_u16(syms_trunc_u32(dw_op.u.line.column));
+        syms_line_init(line_out, (SymsAddr)dw_op.u.line.address, ln, col);
+      } break;
+      case DW_LINE_ITER_OP_DEFINE_FILE: {
+        SYMS_ASSERT_FAILURE("define-file is not implemented");
+        keep_fetching = syms_true;
+      } break;
+      }
+      result = syms_true;
+    }
+  }
+  return result;
+}
+
+SYMS_INTERNAL syms_bool
+syms_member_iter_init_dwarf(SymsDebugInfoDwarf *debug_info, SymsMemberIterDwarf *iter, SymsType *type)
+{
+  const DwTag *type_tag = syms_typeid_to_dw(&type->id);
+  return dw_member_iter_init(&iter->impl, &debug_info->context, *type_tag);
+}
+
+SYMS_INTERNAL syms_bool
+syms_member_iter_next_dwarf(SymsDebugInfoDwarf *debug_info, SymsMemberIterDwarf *iter, SymsMember *member_out)
+{
+  syms_bool result = syms_false;
+  DwMember member;
+(void)debug_info;
+  if (dw_member_iter_next(&iter->impl, &member)) {
+    member_out->type_id  = syms_typeid_for_dw(member.type_tag);
+    member_out->name_ref = syms_string_ref_str(member.name);
+    member_out->u.data_offset = member.byte_off;
+    result = syms_true;
+  }
+  return result;
+}
+
+SYMS_INTERNAL syms_bool
+syms_global_iter_init_dwarf(SymsDebugInfoDwarf *debug_info, SymsGlobalIterDwarf *iter)
+{ (void)debug_info; (void)iter;
+  return syms_false;
+}
+
+SYMS_INTERNAL syms_bool
+syms_global_iter_next_dwarf(SymsDebugInfoDwarf *debug_info, SymsGlobalIterDwarf *iter, SymsGlobal *gdata_out)
+{ (void)debug_info; (void)iter; (void)gdata_out;
+  return syms_false;
+}
+
+SYMS_INTERNAL syms_bool
+syms_local_data_iter_init_dwarf(SymsDebugInfoDwarf *debug_info, SymsLocalDataIterDwarf *iter, SymsMod *mod)
+{ (void)debug_info; (void)iter; (void)mod;
+  return syms_false;
+}
+
+SYMS_INTERNAL syms_bool
+syms_local_data_iter_next_dwarf(SymsDebugInfoDwarf *debug_info, SymsLocalDataIterDwarf *iter, SymsLocalData *ldata_out)
+{ (void)debug_info; (void)iter; (void)ldata_out;
+  return syms_false;
+}
+
+SYMS_INTERNAL syms_bool
+SymsFile_iter_init_dwarf(SymsDebugInfoDwarf *debug_info, SymsFileIterDwarf *iter)
+{
+  return dw_file_iter_init(&iter->impl, &debug_info->context);
+}
+
+SYMS_INTERNAL syms_bool
+SymsFile_iter_next_dwarf(SymsDebugInfoDwarf *debug_info, SymsFileIterDwarf *iter, SymsStringRef *ref_out)
+{ (void)debug_info; (void)iter; (void)ref_out;
+  return syms_false;
+}
+
+SYMS_INTERNAL syms_bool
+syms_proc_iter_init_dwarf(SymsDebugInfoDwarf *debug_info, SymsProcIterDwarf *iter, SymsMod *mod)
+{ (void)debug_info;
+  return dw_proc_iter_init(&iter->impl, (DwCompileUnit *)mod->impl);
+}
+
+SYMS_INTERNAL syms_bool
+syms_proc_iter_next_dwarf(SymsDebugInfoDwarf *debug_info, SymsProcIterDwarf *iter, SymsProc *proc_out)
+{
+  syms_bool result = syms_false;
+  DwProc proc;
+  if (dw_proc_iter_next(&iter->impl, &proc)) {
+    result = syms_proc_from_dw_proc(debug_info, iter->impl.attribs.cu, &proc, proc_out);
+  }
+  return result;
+}
+
+SYMS_INTERNAL syms_bool
+syms_arg_iter_init_dwarf(SymsDebugInfoDwarf *debug_info, SymsArgIterDwarf *iter, SymsTypeID *id)
+{ (void)debug_info; (void)iter; (void)id;
+  return syms_false;
+}
+
+SYMS_INTERNAL syms_bool
+syms_arg_iter_next_dwarf(SymsDebugInfoDwarf *debug_info, SymsArgIterDwarf *iter, SymsTypeID *arg_out)
+{ (void)debug_info; (void)iter; (void)arg_out;
+  return syms_false;
+}
+
+SYMS_INTERNAL syms_bool
+syms_local_iter_init_dwarf(SymsDebugInfoDwarf *debug_info, SymsLocalIterDwarf *iter, SymsProc *proc)
+{
+  SymsProcData *impl = (SymsProcData *)proc->impl; 
+  return dw_local_iter_init(&iter->impl, &debug_info->context, impl->dw.type_tag);
+}
+
+SYMS_INTERNAL syms_bool
+syms_local_iter_next_dwarf(SymsDebugInfoDwarf *debug_info, SymsLocalIterDwarf *iter, SymsVar *lvar_out)
+{
+  syms_bool result = syms_false;
+  DwLocal lvar;
+(void)debug_info;
+  if (dw_local_iter_next(&iter->impl, &lvar)) {
+    lvar_out->type_id = syms_typeid_for_dw(lvar.type_tag);
+    lvar_out->flags = 0;
+    if (lvar.flags & DW_VAR_ARGUMENT) {
+      lvar_out->flags |= SYMS_VAR_FLAG_PARAM;
+    }
+    lvar_out->name_ref = syms_string_ref_str(lvar.name);
+    *(DwEncodedLocation *)lvar_out->encoded_va.impl = lvar.encoded_va;
+    result = syms_true;
+  }
+  return result;
+}
+
+SYMS_INTERNAL syms_bool
+syms_inline_iter_init_dwarf(SymsDebugInfoDwarf *debug_info, SymsInlineIterDwarf *iter, DwCompileUnit *cu, DwTag inlined_proc_tag, SymsAddr rva)
+{ 
+  SymsAddr info_off;
+  syms_bool is_inited;
+(void)debug_info;
+  iter->rva = rva;
+  iter->cu = cu;
+  info_off = DW_PTR_DIFF_BYTES(inlined_proc_tag.info, cu->info_data_start);
+  is_inited = dw_attrib_iter_init(&iter->attribs, cu, info_off);
+  is_inited = is_inited && iter->attribs.has_children;
+  return is_inited;
+}
+
+SYMS_INTERNAL syms_bool
+syms_inline_iter_next_dwarf(SymsDebugInfoDwarf *debug_info, SymsInlineIterDwarf *iter, SymsInlineSite *site_out)
+{
+  syms_bool is_next_valid = syms_false;
+  DwCompileUnit *cu = iter->cu;
+(void)debug_info;
+
+  do {
+    if (iter->attribs.has_children) {
+      iter->depth += 1;
+    }
+    if (iter->attribs.tag_type == DW_TAG_INLINED_SUBROUTINE) {
+      DwProc dw_proc;
+      if (dw_proc_init(&iter->attribs, &dw_proc)) {
+        SymsAddr site_lo = 0;
+        SymsAddr site_hi = 0;
+        if (!dw_range_check(iter->cu, dw_proc.range_off, iter->rva, &site_lo, &site_hi)) {
+          DwLocation loc;
+          if (dw_decode_location_rva(&dw_proc.encoded_va.u.rva, 0, &loc)) {
+            if (loc.type == DW_LOCATION_ADDR) {
+              site_lo = loc.u.addr;
+              site_hi = site_lo + dw_proc.len;
+            } else {
+#if defined(SYMS_PARANOID)
+              SYMS_ASSERT_FAILURE("expected an address for inline site");
+#endif
+            }
+          }
+        }
+
+        is_next_valid = iter->rva >= site_lo && iter->rva < site_hi;
+        if (is_next_valid) {
+          DwLineIter line_iter;
+
+          site_out->name = syms_string_ref_str(dw_proc.name);
+          site_out->type_id = syms_typeid_for_dw(dw_proc.type_tag);
+
+          site_out->call_file.name = syms_string_ref_str(syms_string_init(0,0));
+          site_out->call_file.chksum_type = SYMS_CHECKSUM_NULL;
+
+          site_out->decl_file.name = syms_string_ref_str(syms_string_init(0,0));
+          site_out->decl_file.chksum_type = SYMS_CHECKSUM_NULL;
+
+          site_out->call_ln = dw_proc.call_ln;
+          site_out->decl_ln = dw_proc.decl_ln;
+
+          site_out->range_lo = site_lo;
+          site_out->range_hi = site_hi;
+
+          site_out->sort_index = (cu->info_base + cu->info_len) - iter->attribs.info_off;
+
+          if (dw_line_iter_init(&line_iter, iter->cu)) {
+            DwLineFile file;
+
+            // resolve call file and line number
+            if (dw_line_iter_get_file(&line_iter, dw_proc.call_file, &file)) {
+              SymsString dir;
+              if (dw_line_iter_get_dir(&line_iter, file.dir_index, &dir)) {
+                site_out->call_file.name = syms_string_ref_dw_path(cu->compile_dir, dir, file.file_name);
+              } else {
+                site_out->call_file.name = syms_string_ref_str(file.file_name);
+              }
+            }
+
+            // resolve decl file and line number
+            if (dw_line_iter_get_file(&line_iter, dw_proc.decl_file, &file)) {
+              SymsString dir;
+              if (dw_line_iter_get_dir(&line_iter, file.dir_index, &dir)) {
+                site_out->decl_file.name = syms_string_ref_dw_path(cu->compile_dir, dir, file.file_name);
+              } else {
+                site_out->decl_file.name = syms_string_ref_str(file.file_name);
+              }
+            }
+          }
+          break;
+        }
+      }
+    } else if (iter->attribs.tag_type == DW_TAG_NULL) {
+      if (iter->depth == 0) {
+        break;
+      }
+      iter->depth -= 1;
+    }
+  } while (dw_attrib_iter_next_tag(&iter->attribs));
+
+  dw_attrib_iter_next_tag(&iter->attribs);
+
+  return is_next_valid;
+}
+
+SYMS_INTERNAL syms_bool
+syms_const_iter_init_dwarf(SymsDebugInfoDwarf *debug_info, SymsConstIterDwarf *iter)
+{ (void)debug_info; (void)iter;
+  return syms_false;
+}
+
+SYMS_INTERNAL syms_bool
+syms_const_iter_next_dwarf(SymsDebugInfoDwarf *debug_info, SymsConstIterDwarf *iter, SymsConst *const_out)
+{ (void)debug_info; (void)iter; (void)const_out;
+  return syms_false;
+}
+
+SYMS_INTERNAL syms_bool
+syms_type_iter_init_dwarf(SymsDebugInfoDwarf *debug_info, SymsTypeIterDwarf *iter)
+{ (void)debug_info; (void)iter;
+  return syms_false;
+}
+
+SYMS_INTERNAL syms_bool
+syms_type_iter_next_dwarf(SymsDebugInfoDwarf *debug_info, SymsTypeIterDwarf *iter, SymsTypeID *typeid_out)
+{ (void)debug_info; (void)iter; (void)typeid_out;
+  return syms_false;
+}
+
+SYMS_INTERNAL syms_bool
+syms_type_from_name_dwarf(SymsDebugInfoDwarf *debug_info, SymsString name, SymsType *type_out)
+{
+  syms_bool result = syms_false;
+  DwTag type_tag;
+  if (dw_type_from_name(&debug_info->context, name.data, name.len, 1, &type_tag)) {
+    DwType type;
+    if (dw_infer_type(&debug_info->context, type_tag, &type)) {
+      result = syms_dw_type_to_syms_type(&type, type_out);
+    }
+  }
+  return result;
+}
+
+SYMS_INTERNAL syms_bool
+syms_global_from_name_dwarf(SymsDebugInfoDwarf *debug_info, SymsString name, SymsGlobal *gvar_out)
+{
+  syms_bool result = syms_false;
+  DwVar gvar;
+  if (dw_global_from_name(&debug_info->context, name.data, name.len, &gvar)) {
+    gvar_out->type_id = syms_typeid_for_dw(gvar.type_tag);
+    *(DwEncodedLocation *)gvar_out->encoded_va.impl = gvar.encoded_va;
+    result = syms_true;
+  }
+  return result;
+}
+
+SYMS_INTERNAL syms_bool
+syms_const_from_name_dwarf(SymsDebugInfoDwarf *debug_info, SymsString name, SymsConst *const_out)
+{ (void)debug_info; (void)name; (void)const_out;
+  return syms_false;
+}
+
+SYMS_INTERNAL syms_bool
+syms_infer_type_dwarf(SymsDebugInfoDwarf *debug_info, SymsTypeID type_id, SymsType *type_out)
+{
+  DwType type;
+  const DwTag *type_tag = syms_typeid_to_dw(&type_id);
+  syms_bool result = syms_false;
+  if (dw_infer_type(&debug_info->context, *type_tag, &type)) {
+    result = syms_dw_type_to_syms_type(&type, type_out);
+  }
+  return result;
+}
+
+SYMS_INTERNAL syms_bool
+syms_trampoline_from_ip_dwarf(SymsDebugInfoDwarf *debug_info, SymsAddr ip, SymsAddr *ip_out)
+{ (void)debug_info; (void)ip; (void)ip_out;
+  return syms_false;
+}
+
+SYMS_INTERNAL SymsErrorCode
+syms_decode_location_dwarf(
+               SymsAddr orig_rebase,
+               SymsAddr rebase,
+               DwEncodedLocation *impl, 
+               SymsRegread *regread,
+               SymsMemread *memread_context,
+               SymsLocation *loc_out)
+{
+  SymsErrorCode result = SYMS_ERR_INVALID_CODE_PATH;
+  DwLocation loc;
+  if (dw_decode_location(impl, rebase, memread_context, syms_memread_dwarf, regread, syms_regread_dwarf, &loc)) {
+    if (SYMS_RESULT_OK(memread_context->result)) {
+      switch (loc.type) {
+      case DW_LOCATION_ADDR: {
+        SYMS_ASSERT(loc.u.addr >= orig_rebase);
+        loc_out->kind = SYMS_LOCATION_VA;
+        loc_out->u.va = (loc.u.addr - orig_rebase) + rebase;
+        result = SYMS_ERR_OK;
+      } break;
+      case DW_LOCATION_IMPLICIT: {
+        loc_out->kind = SYMS_LOCATION_INDIRECT;
+        loc_out->u.indirect.len = loc.u.implicit.len;
+        loc_out->u.indirect.data = loc.u.implicit.data;
+        result = SYMS_ERR_OK;
+      } break;
+      case DW_LOCATION_NULL: {
+        loc_out->kind = SYMS_LOCATION_NULL;
+        result = SYMS_ERR_OK;
+      } break;
+      }
+    } else {
+      result = memread_context->result;
+    }
+  }
+  return result;
+}
+
+SYMS_INTERNAL syms_bool
+syms_proc_from_dw_proc(SymsDebugInfoDwarf *debug_info, DwCompileUnit *cu, DwProc *proc, SymsProc *proc_out)
+{
+  DwLocation location;
+(void)debug_info;
+  if (proc_out && proc) {
+    SymsAddr range_lo, range_hi;
+
+    proc_out->type_id               = syms_typeid_for_dw(proc->type_tag);
+    proc_out->len                   = proc->len;
+    proc_out->name_ref              = syms_string_ref_str(proc->name);
+    {
+      SymsProcData *impl = (SymsProcData *)proc_out->impl;
+      impl->dw.type_tag      = proc->type_tag;
+      impl->dw.frame_base    = proc->frame_base;
+    }
+
+    {
+      SymsRange *range = &proc_out->range;
+      SymsRangeDwarf *impl = (SymsRangeDwarf *)range->u.impl;
+      range->type = SYMS_RANGE_IMPL;
+      impl->cu_index = cu->index;
+      impl->range_off = proc->range_off;
+    }
+
+    if (dw_get_range_bounds(cu, proc->range_off, &range_lo, &range_hi)) {
+      proc_out->va = range_lo;
+      proc_out->len = syms_trunc_u32(range_hi - range_lo);
+    }
+
+    if (dw_decode_location(&proc->encoded_va, 0, NULL, NULL, NULL, NULL, &location)) {
+      if (location.type == DW_LOCATION_ADDR) {
+        proc_out->va            = location.u.addr;
+        proc_out->dbg_start_va  = proc_out->va;
+        proc_out->dbg_end_va    = proc_out->va + proc_out->len;
+      } else {
+        proc_out->va            = 0;
+        proc_out->dbg_start_va  = 0;
+        proc_out->dbg_end_va    = 0;
+      }
+    }
+  }
+  return syms_true;
+}
+
+SYMS_INTERNAL syms_bool
+syms_dw_type_to_syms_type(DwType *dw_type, SymsType *syms_type)
+{
+  syms_type->id = syms_typeid_for_dw(dw_type->type_tag);
+  syms_type->next_id = syms_typeid_for_dw(dw_type->next_type_tag);
+  syms_type->is_fwdref = syms_false;
+
+  syms_type->modifier = 0;
+  if (dw_type->modifier & DW_TYPE_MDFR_ATOMIC) {
+    syms_type->modifier |= SYMS_TYPE_MDFR_ATOMIC;
+  }
+  if (dw_type->modifier & DW_TYPE_MDFR_CONST) {
+    syms_type->modifier |= SYMS_TYPE_MDFR_CONST;
+  }
+  if (dw_type->modifier & DW_TYPE_MDFR_IMMUTABLE) {
+    syms_type->modifier |= SYMS_TYPE_MDFR_IMMUTABLE;
+  }
+  if (dw_type->modifier & DW_TYPE_MDFR_PACKED) {
+    syms_type->modifier |= SYMS_TYPE_MDFR_PACKED;
+  }
+  if (dw_type->modifier & DW_TYPE_MDFR_REF) {
+    syms_type->modifier |= SYMS_TYPE_MDFR_REF;
+  }
+  if (dw_type->modifier & DW_TYPE_MDFR_RESTRICT) {
+    syms_type->modifier |= SYMS_TYPE_MDFR_RESTRICT;
+  }
+  if (dw_type->modifier & DW_TYPE_MDFR_RVALUE_REF) {
+    syms_type->modifier |= SYMS_TYPE_MDFR_REF;
+  }
+  if (dw_type->modifier & DW_TYPE_MDFR_SHARED) {
+    syms_type->modifier |= SYMS_TYPE_MDFR_SHARED;
+  }
+  if (dw_type->modifier & DW_TYPE_MDFR_VOLATILE) {
+    syms_type->modifier |= SYMS_TYPE_MDFR_VOLATILE;
+  }
+  if (dw_type->modifier & DW_TYPE_MDFR_CHAR) {
+    syms_type->modifier |= SYMS_TYPE_MDFR_CHAR;
+  }
+  
+  switch (dw_type->kind) {
+  case DW_TYPE_NULL:      syms_type->kind = SYMS_TYPE_NULL; break;
+
+  case DW_TYPE_FLOAT16:   syms_type->kind = SYMS_TYPE_FLOAT16;    break;
+  case DW_TYPE_FLOAT32:   syms_type->kind = SYMS_TYPE_FLOAT32;    break;
+  case DW_TYPE_FLOAT48:   syms_type->kind = SYMS_TYPE_FLOAT48;    break;
+  case DW_TYPE_FLOAT64:   syms_type->kind = SYMS_TYPE_FLOAT64;    break;
+  case DW_TYPE_FLOAT80:   syms_type->kind = SYMS_TYPE_FLOAT80;    break;
+  case DW_TYPE_FLOAT128:  syms_type->kind = SYMS_TYPE_FLOAT128;   break;
+
+  case DW_TYPE_INT8:      syms_type->kind = SYMS_TYPE_INT8;   break;
+  case DW_TYPE_INT16:     syms_type->kind = SYMS_TYPE_INT16;  break;
+  case DW_TYPE_INT32:     syms_type->kind = SYMS_TYPE_INT32;  break;
+  case DW_TYPE_INT64:     syms_type->kind = SYMS_TYPE_INT64;  break;
+  case DW_TYPE_INT128:    syms_type->kind = SYMS_TYPE_INT128; break;
+
+  case DW_TYPE_UINT8:      syms_type->kind = SYMS_TYPE_UINT8;   break;
+  case DW_TYPE_UINT16:     syms_type->kind = SYMS_TYPE_UINT16;  break;
+  case DW_TYPE_UINT32:     syms_type->kind = SYMS_TYPE_UINT32;  break;
+  case DW_TYPE_UINT64:     syms_type->kind = SYMS_TYPE_UINT64;  break;
+  case DW_TYPE_UINT128:    syms_type->kind = SYMS_TYPE_UINT128; break;
+
+  case DW_TYPE_STRUCT:     syms_type->kind = SYMS_TYPE_STRUCT;     break;
+  case DW_TYPE_UNION:      syms_type->kind = SYMS_TYPE_UNION;      break;
+  case DW_TYPE_CLASS:      syms_type->kind = SYMS_TYPE_CLASS;      break;
+  case DW_TYPE_TYPEDEF:    syms_type->kind = SYMS_TYPE_TYPEDEF;    break;
+  case DW_TYPE_ENUM:       syms_type->kind = SYMS_TYPE_ENUM;       break;
+  case DW_TYPE_PROC:       syms_type->kind = SYMS_TYPE_PROC;       break;
+  case DW_TYPE_PROC_PARAM: syms_type->kind = SYMS_TYPE_PROC_PARAM; break;
+  case DW_TYPE_VOID:       syms_type->kind = SYMS_TYPE_VOID;       break;
+  case DW_TYPE_BOOL:       syms_type->kind = SYMS_TYPE_BOOL;       break;
+  case DW_TYPE_PTR:        syms_type->kind = SYMS_TYPE_PTR;        break;
+  case DW_TYPE_ARR:        syms_type->kind = SYMS_TYPE_ARR;        break;
+
+#if 0
+  case DW_TYPE_CHAR:       syms_type->kind = SYMS_TYPE_CHAR;       break;
+  case DW_TYPE_METHOD:     syms_type->kind = SYMS_TYPE_METHOD;     break;
+  case DW_TYPE_COMPLEX32:  syms_type->kind = SYMS_TYPE_COMPLEX32;  break;
+  case DW_TYPE_COMPLEX64:  syms_type->kind = SYMS_TYPE_COMPLEX64;  break;
+  case DW_TYPE_COMPLEX80:  syms_type->kind = SYMS_TYPE_COMPLEX80;  break;
+  case DW_TYPE_COMPLEX128: syms_type->kind = SYMS_TYPE_COMPLEX128; break;
+#endif
+  }
+  
+  if (dw_type->kind == DW_TYPE_ARR) {
+    syms_type->size = (syms_uint)dw_type->u.arr_count;
+  } else {
+    syms_type->size = (syms_uint)dw_type->size;
+  }
+  syms_type->name_ref = syms_string_ref_str(dw_type->name);
+
+  return syms_true;
+}
+
+SYMS_INTERNAL SymsStringRef
+syms_string_ref_dw_path(SymsString compile_dir, SymsString dir, SymsString file)
+{
+  SymsStringRef ref;
+  DwFilePath *path;
+  ref.type = SYMS_STRING_REF_TYPE_DW_PATH;
+  path = (DwFilePath *)ref.impl;
+  path->compile_dir = compile_dir;
+  path->dir = dir;
+  path->file = file;
+  return ref;
+}
+
+SYMS_INTERNAL syms_bool
+syms_dw_regid_to_regid(SymsArch arch, dw_uint regid, SymsRegID *regid_out)
+{
+  syms_bool is_mapped = syms_true;
+
+  switch (arch) {
+  case SYMS_ARCH_X64: {
+    switch (regid) {
+    case DW_REG_X64_RAX:        *regid_out = SYMS_REG_X64_rax;     break;
+    case DW_REG_X64_RDX:        *regid_out = SYMS_REG_X64_rdx;     break;
+    case DW_REG_X64_RCX:        *regid_out = SYMS_REG_X64_rcx;     break;
+    case DW_REG_X64_RBX:        *regid_out = SYMS_REG_X64_rbx;     break;
+    case DW_REG_X64_RSI:        *regid_out = SYMS_REG_X64_rsi;     break;
+    case DW_REG_X64_RDI:        *regid_out = SYMS_REG_X64_rdi;     break;
+    case DW_REG_X64_RBP:        *regid_out = SYMS_REG_X64_rbp;     break;
+    case DW_REG_X64_RSP:        *regid_out = SYMS_REG_X64_rsp;     break;
+    case DW_REG_X64_R8:         *regid_out = SYMS_REG_X64_r8;      break;
+    case DW_REG_X64_R9:         *regid_out = SYMS_REG_X64_r9;      break;
+    case DW_REG_X64_R10:        *regid_out = SYMS_REG_X64_r10;     break;
+    case DW_REG_X64_R11:        *regid_out = SYMS_REG_X64_r11;     break;
+    case DW_REG_X64_R12:        *regid_out = SYMS_REG_X64_r12;     break;
+    case DW_REG_X64_R13:        *regid_out = SYMS_REG_X64_r13;     break;
+    case DW_REG_X64_R14:        *regid_out = SYMS_REG_X64_r14;     break;
+    case DW_REG_X64_R15:        *regid_out = SYMS_REG_X64_r15;     break;
+    case DW_REG_X64_RIP:        *regid_out = SYMS_REG_X64_rip;     break;
+    case DW_REG_X64_XMM0:       *regid_out = SYMS_REG_X64_xmm0;    break;
+    case DW_REG_X64_XMM1:       *regid_out = SYMS_REG_X64_xmm1;    break;
+    case DW_REG_X64_XMM2:       *regid_out = SYMS_REG_X64_xmm2;    break;
+    case DW_REG_X64_XMM3:       *regid_out = SYMS_REG_X64_xmm3;    break;
+    case DW_REG_X64_XMM4:       *regid_out = SYMS_REG_X64_xmm4;    break;
+    case DW_REG_X64_XMM5:       *regid_out = SYMS_REG_X64_xmm5;    break;
+    case DW_REG_X64_XMM6:       *regid_out = SYMS_REG_X64_xmm6;    break;
+    case DW_REG_X64_XMM7:       *regid_out = SYMS_REG_X64_xmm7;    break;
+    case DW_REG_X64_XMM8:       *regid_out = SYMS_REG_X64_xmm8;    break;
+    case DW_REG_X64_XMM9:       *regid_out = SYMS_REG_X64_xmm9;    break;
+    case DW_REG_X64_XMM10:      *regid_out = SYMS_REG_X64_xmm10;   break;
+    case DW_REG_X64_XMM11:      *regid_out = SYMS_REG_X64_xmm11;   break;
+    case DW_REG_X64_XMM12:      *regid_out = SYMS_REG_X64_xmm12;   break;
+    case DW_REG_X64_XMM13:      *regid_out = SYMS_REG_X64_xmm13;   break;
+    case DW_REG_X64_XMM14:      *regid_out = SYMS_REG_X64_xmm14;   break;
+    case DW_REG_X64_XMM15:      *regid_out = SYMS_REG_X64_xmm15;   break;
+    case DW_REG_X64_ST0:        *regid_out = SYMS_REG_X64_st0;     break;
+    case DW_REG_X64_ST1:        *regid_out = SYMS_REG_X64_st1;     break;
+    case DW_REG_X64_ST2:        *regid_out = SYMS_REG_X64_st2;     break;
+    case DW_REG_X64_ST3:        *regid_out = SYMS_REG_X64_st3;     break;
+    case DW_REG_X64_ST4:        *regid_out = SYMS_REG_X64_st4;     break;
+    case DW_REG_X64_ST5:        *regid_out = SYMS_REG_X64_st5;     break;
+    case DW_REG_X64_ST6:        *regid_out = SYMS_REG_X64_st6;     break;
+    case DW_REG_X64_ST7:        *regid_out = SYMS_REG_X64_st7;     break;
+    case DW_REG_X64_MM0:        *regid_out = SYMS_REG_X64_mm0;     break;
+    case DW_REG_X64_MM1:        *regid_out = SYMS_REG_X64_mm1;     break;
+    case DW_REG_X64_MM2:        *regid_out = SYMS_REG_X64_mm2;     break;
+    case DW_REG_X64_MM3:        *regid_out = SYMS_REG_X64_mm3;     break;
+    case DW_REG_X64_MM4:        *regid_out = SYMS_REG_X64_mm4;     break;
+    case DW_REG_X64_MM5:        *regid_out = SYMS_REG_X64_mm5;     break;
+    case DW_REG_X64_MM6:        *regid_out = SYMS_REG_X64_mm6;     break;
+    case DW_REG_X64_MM7:        *regid_out = SYMS_REG_X64_mm7;     break;
+    case DW_REG_X64_RFLAGS:     *regid_out = SYMS_REG_X64_rflags;  break;
+    case DW_REG_X64_ES:         *regid_out = SYMS_REG_X64_es;      break;
+    case DW_REG_X64_CS:         *regid_out = SYMS_REG_X64_cs;      break;
+    case DW_REG_X64_SS:         *regid_out = SYMS_REG_X64_ss;      break;
+    case DW_REG_X64_DS:         *regid_out = SYMS_REG_X64_ds;      break;
+    case DW_REG_X64_FS:         *regid_out = SYMS_REG_X64_fs;      break;
+    case DW_REG_X64_GS:         *regid_out = SYMS_REG_X64_gs;      break;
+    case DW_REG_X64_FS_BASE:    *regid_out = SYMS_REG_X64_fsbase;  break;
+    case DW_REG_X64_GS_BASE:    *regid_out = SYMS_REG_X64_gsbase;  break;
+
+    case DW_REG_X64_TR:
+    case DW_REG_X64_LDTR:
+    default: is_mapped = syms_false; break;
+    }
+  } break;
+
+  case SYMS_ARCH_X86: {
+    switch (regid) {
+    case DW_REG_X86_EAX:     *regid_out = SYMS_REG_X86_eax;    break;
+    case DW_REG_X86_ECX:     *regid_out = SYMS_REG_X86_ecx;    break;
+    case DW_REG_X86_EDX:     *regid_out = SYMS_REG_X86_edx;    break;
+    case DW_REG_X86_EBX:     *regid_out = SYMS_REG_X86_ebx;    break;
+    case DW_REG_X86_ESP:     *regid_out = SYMS_REG_X86_esp;    break;
+    case DW_REG_X86_EBP:     *regid_out = SYMS_REG_X86_ebp;    break;
+    case DW_REG_X86_ESI:     *regid_out = SYMS_REG_X86_esi;    break;
+    case DW_REG_X86_EDI:     *regid_out = SYMS_REG_X86_edi;    break;
+    case DW_REG_X86_EIP:     *regid_out = SYMS_REG_X86_eip;    break;
+    case DW_REG_X86_EFLAGS:  *regid_out = SYMS_REG_X86_eflags; break;
+    case DW_REG_X86_ST0:     *regid_out = SYMS_REG_X86_st0;    break;
+    case DW_REG_X86_ST1:     *regid_out = SYMS_REG_X86_st1;    break;
+    case DW_REG_X86_ST2:     *regid_out = SYMS_REG_X86_st2;    break;
+    case DW_REG_X86_ST3:     *regid_out = SYMS_REG_X86_st3;    break;
+    case DW_REG_X86_ST4:     *regid_out = SYMS_REG_X86_st4;    break;
+    case DW_REG_X86_ST5:     *regid_out = SYMS_REG_X86_st5;    break;
+    case DW_REG_X86_ST6:     *regid_out = SYMS_REG_X86_st6;    break;
+    case DW_REG_X86_ST7:     *regid_out = SYMS_REG_X86_st7;    break;
+    case DW_REG_X86_XMM0:    *regid_out = SYMS_REG_X86_xmm0;   break;
+    case DW_REG_X86_XMM1:    *regid_out = SYMS_REG_X86_xmm1;   break;
+    case DW_REG_X86_XMM2:    *regid_out = SYMS_REG_X86_xmm2;   break;
+    case DW_REG_X86_XMM3:    *regid_out = SYMS_REG_X86_xmm3;   break;
+    case DW_REG_X86_XMM4:    *regid_out = SYMS_REG_X86_xmm4;   break;
+    case DW_REG_X86_XMM5:    *regid_out = SYMS_REG_X86_xmm5;   break;
+    case DW_REG_X86_XMM6:    *regid_out = SYMS_REG_X86_xmm6;   break;
+    case DW_REG_X86_XMM7:    *regid_out = SYMS_REG_X86_xmm7;   break;
+    case DW_REG_X86_MM0:     *regid_out = SYMS_REG_X86_mm0;    break;
+    case DW_REG_X86_MM1:     *regid_out = SYMS_REG_X86_mm1;    break;
+    case DW_REG_X86_MM2:     *regid_out = SYMS_REG_X86_mm2;    break;
+    case DW_REG_X86_MM3:     *regid_out = SYMS_REG_X86_mm3;    break;
+    case DW_REG_X86_MM4:     *regid_out = SYMS_REG_X86_mm4;    break;
+    case DW_REG_X86_MM5:     *regid_out = SYMS_REG_X86_mm5;    break;
+    case DW_REG_X86_MM6:     *regid_out = SYMS_REG_X86_mm6;    break;
+    case DW_REG_X86_MM7:     *regid_out = SYMS_REG_X86_mm7;    break;
+    case DW_REG_X86_FCW:     *regid_out = SYMS_REG_X86_fcw;    break;
+    case DW_REG_X86_FSW:     *regid_out = SYMS_REG_X86_fsw;    break;
+    case DW_REG_X86_MXCSR:   *regid_out = SYMS_REG_X86_mxcsr;  break;
+    case DW_REG_X86_ES:      *regid_out = SYMS_REG_X86_es;     break;
+    case DW_REG_X86_CS:      *regid_out = SYMS_REG_X86_cs;     break;
+    case DW_REG_X86_SS:      *regid_out = SYMS_REG_X86_ss;     break;
+    case DW_REG_X86_DS:      *regid_out = SYMS_REG_X86_ds;     break;
+    case DW_REG_X86_FS:      *regid_out = SYMS_REG_X86_fs;     break;
+    case DW_REG_X86_GS:      *regid_out = SYMS_REG_X86_gs;     break;
+
+    case DW_REG_X86_TRAPNO:
+    case DW_REG_X86_TR:
+    case DW_REG_X86_LDTR:
+    default: {
+      is_mapped = syms_false;
+    } break;
+    }
+  } break;
+
+  case SYMS_ARCH_ARM: {
+    // TODO(nick): I haven't thought too hard about the ARM support,
+    // so this is just a stub code.
+#if 0
+    switch (regid) {
+    case DW_REG_ARM_S0:       *regid_out = SYMS_REG_ARM_s0;       break;
+    case DW_REG_ARM_S1:       *regid_out = SYMS_REG_ARM_s1;       break;
+    case DW_REG_ARM_S2:       *regid_out = SYMS_REG_ARM_s2;       break;
+    case DW_REG_ARM_S3:       *regid_out = SYMS_REG_ARM_s3;       break;
+    case DW_REG_ARM_S4:       *regid_out = SYMS_REG_ARM_s4;       break;
+    case DW_REG_ARM_S5:       *regid_out = SYMS_REG_ARM_s5;       break;
+    case DW_REG_ARM_S6:       *regid_out = SYMS_REG_ARM_s6;       break;
+    case DW_REG_ARM_S7:       *regid_out = SYMS_REG_ARM_s7;       break;
+    case DW_REG_ARM_S8:       *regid_out = SYMS_REG_ARM_s8;       break;
+    case DW_REG_ARM_S9:       *regid_out = SYMS_REG_ARM_s9;       break;
+    case DW_REG_ARM_S10:      *regid_out = SYMS_REG_ARM_s10;      break;
+    case DW_REG_ARM_S11:      *regid_out = SYMS_REG_ARM_s11;      break;
+    case DW_REG_ARM_S12:      *regid_out = SYMS_REG_ARM_s12;      break;
+    case DW_REG_ARM_S13:      *regid_out = SYMS_REG_ARM_s13;      break;
+    case DW_REG_ARM_S14:      *regid_out = SYMS_REG_ARM_s14;      break;
+    case DW_REG_ARM_S15:      *regid_out = SYMS_REG_ARM_s15;      break;
+    case DW_REG_ARM_S16:      *regid_out = SYMS_REG_ARM_s16;      break;
+    case DW_REG_ARM_S17:      *regid_out = SYMS_REG_ARM_s17;      break;
+    case DW_REG_ARM_S18:      *regid_out = SYMS_REG_ARM_s18;      break;
+    case DW_REG_ARM_S19:      *regid_out = SYMS_REG_ARM_s19;      break;
+    case DW_REG_ARM_S20:      *regid_out = SYMS_REG_ARM_s20;      break;
+    case DW_REG_ARM_S21:      *regid_out = SYMS_REG_ARM_s21;      break;
+    case DW_REG_ARM_S22:      *regid_out = SYMS_REG_ARM_s22;      break;
+    case DW_REG_ARM_S23:      *regid_out = SYMS_REG_ARM_s23;      break;
+    case DW_REG_ARM_S24:      *regid_out = SYMS_REG_ARM_s24;      break;
+    case DW_REG_ARM_S25:      *regid_out = SYMS_REG_ARM_s25;      break;
+    case DW_REG_ARM_S26:      *regid_out = SYMS_REG_ARM_s26;      break;
+    case DW_REG_ARM_S27:      *regid_out = SYMS_REG_ARM_s27;      break;
+    case DW_REG_ARM_S28:      *regid_out = SYMS_REG_ARM_s28;      break;
+    case DW_REG_ARM_S29:      *regid_out = SYMS_REG_ARM_s29;      break;
+    case DW_REG_ARM_S30:      *regid_out = SYMS_REG_ARM_s30;      break;
+    case DW_REG_ARM_S31:      *regid_out = SYMS_REG_ARM_s31;      break;
+    case DW_REG_ARM_F0:       *regid_out = SYMS_REG_ARM_f0;       break;
+    case DW_REG_ARM_F1:       *regid_out = SYMS_REG_ARM_f1;       break;
+    case DW_REG_ARM_F2:       *regid_out = SYMS_REG_ARM_f2;       break;
+    case DW_REG_ARM_F3:       *regid_out = SYMS_REG_ARM_f3;       break;
+    case DW_REG_ARM_F4:       *regid_out = SYMS_REG_ARM_f4;       break;
+    case DW_REG_ARM_F5:       *regid_out = SYMS_REG_ARM_f5;       break;
+    case DW_REG_ARM_F6:       *regid_out = SYMS_REG_ARM_f6;       break;
+    case DW_REG_ARM_F7:       *regid_out = SYMS_REG_ARM_f7;       break;
+    case DW_REG_ARM_WCGR0:    *regid_out = SYMS_REG_ARM_wcgr0;    break;
+    case DW_REG_ARM_WCGR1:    *regid_out = SYMS_REG_ARM_wcgr1;    break;
+    case DW_REG_ARM_WCGR2:    *regid_out = SYMS_REG_ARM_wcgr2;    break;
+    case DW_REG_ARM_WCGR3:    *regid_out = SYMS_REG_ARM_wcgr3;    break;
+    case DW_REG_ARM_WCGR4:    *regid_out = SYMS_REG_ARM_wcgr4;    break;
+    case DW_REG_ARM_WCGR5:    *regid_out = SYMS_REG_ARM_wcgr5;    break;
+    case DW_REG_ARM_WCGR6:    *regid_out = SYMS_REG_ARM_wcgr6;    break;
+    case DW_REG_ARM_WCGR7:    *regid_out = SYMS_REG_ARM_wcgr7;    break;
+    case DW_REG_ARM_WR0:      *regid_out = SYMS_REG_ARM_wr0;      break;
+    case DW_REG_ARM_WR1:      *regid_out = SYMS_REG_ARM_wr1;      break;
+    case DW_REG_ARM_WR2:      *regid_out = SYMS_REG_ARM_wr2;      break;
+    case DW_REG_ARM_WR3:      *regid_out = SYMS_REG_ARM_wr3;      break;
+    case DW_REG_ARM_WR4:      *regid_out = SYMS_REG_ARM_wr4;      break;
+    case DW_REG_ARM_WR5:      *regid_out = SYMS_REG_ARM_wr5;      break;
+    case DW_REG_ARM_WR6:      *regid_out = SYMS_REG_ARM_wr6;      break;
+    case DW_REG_ARM_WR7:      *regid_out = SYMS_REG_ARM_wr7;      break;
+    case DW_REG_ARM_WR8:      *regid_out = SYMS_REG_ARM_wr8;      break;
+    case DW_REG_ARM_WR9:      *regid_out = SYMS_REG_ARM_wr9;      break;
+    case DW_REG_ARM_WR10:     *regid_out = SYMS_REG_ARM_wr10;     break;
+    case DW_REG_ARM_WR11:     *regid_out = SYMS_REG_ARM_wr11;     break;
+    case DW_REG_ARM_WR12:     *regid_out = SYMS_REG_ARM_wr12;     break;
+    case DW_REG_ARM_WR13:     *regid_out = SYMS_REG_ARM_wr13;     break;
+    case DW_REG_ARM_WR14:     *regid_out = SYMS_REG_ARM_wr14;     break;
+    case DW_REG_ARM_WR15:     *regid_out = SYMS_REG_ARM_wr15;     break;
+    case DW_REG_ARM_SPSR:     *regid_out = SYMS_REG_ARM_spsr;     break;
+    case DW_REG_ARM_spsr_FIQ: *regid_out = SYMS_REG_ARM_spsr_fiq; break;
+    case DW_REG_ARM_spsr_IRQ: *regid_out = SYMS_REG_ARM_spsr_irq; break;
+    case DW_REG_ARM_spsr_ABT: *regid_out = SYMS_REG_ARM_spsr_abt; break;
+    case DW_REG_ARM_spsr_UND: *regid_out = SYMS_REG_ARM_spsr_und; break;
+    case DW_REG_ARM_spsr_SVC: *regid_out = SYMS_REG_ARM_spsr_svc; break;
+    case DW_REG_ARM_r8_USR:   *regid_out = SYMS_REG_ARM_r8_usr;   break;
+    case DW_REG_ARM_r9_USR:   *regid_out = SYMS_REG_ARM_r9_usr;   break;
+    case DW_REG_ARM_r10_USR:  *regid_out = SYMS_REG_ARM_r10_usr;  break;
+    case DW_REG_ARM_r11_USR:  *regid_out = SYMS_REG_ARM_r11_usr;  break;
+    case DW_REG_ARM_r12_USR:  *regid_out = SYMS_REG_ARM_r12_usr;  break;
+    case DW_REG_ARM_r13_USR:  *regid_out = SYMS_REG_ARM_r13_usr;  break;
+    case DW_REG_ARM_r14_USR:  *regid_out = SYMS_REG_ARM_r14_usr;  break;
+    case DW_REG_ARM_r8_FIQ:   *regid_out = SYMS_REG_ARM_r8_fiq;   break;
+    case DW_REG_ARM_r9_FIQ:   *regid_out = SYMS_REG_ARM_r9_fiq;   break;
+    case DW_REG_ARM_r10_FIQ:  *regid_out = SYMS_REG_ARM_r10_fiq;  break;
+    case DW_REG_ARM_r11_FIQ:  *regid_out = SYMS_REG_ARM_r11_fiq;  break;
+    case DW_REG_ARM_r12_FIQ:  *regid_out = SYMS_REG_ARM_r12_fiq;  break;
+    case DW_REG_ARM_r13_FIQ:  *regid_out = SYMS_REG_ARM_r13_fiq;  break;
+    case DW_REG_ARM_r14_FIQ:  *regid_out = SYMS_REG_ARM_r14_fiq;  break;
+    case DW_REG_ARM_r13_RIQ:  *regid_out = SYMS_REG_ARM_r13_riq;  break;
+    case DW_REG_ARM_r14_RIQ:  *regid_out = SYMS_REG_ARM_r14_riq;  break;
+    case DW_REG_ARM_r14_ABT:  *regid_out = SYMS_REG_ARM_r14_abt;  break;
+    case DW_REG_ARM_r13_ABT:  *regid_out = SYMS_REG_ARM_r13_abt;  break;
+    case DW_REG_ARM_r14_UND:  *regid_out = SYMS_REG_ARM_r14_und;  break;
+    case DW_REG_ARM_r13_UND:  *regid_out = SYMS_REG_ARM_r13_und;  break;
+    case DW_REG_ARM_r14_SVC:  *regid_out = SYMS_REG_ARM_r14_svc;  break;
+    case DW_REG_ARM_r13_SVC:  *regid_out = SYMS_REG_ARM_r13_svc;  break;
+    case DW_REG_ARM_WC0:      *regid_out = SYMS_REG_ARM_wc0;      break;
+    case DW_REG_ARM_WC1:      *regid_out = SYMS_REG_ARM_wc1;      break;
+    case DW_REG_ARM_WC2:      *regid_out = SYMS_REG_ARM_wc2;      break;
+    case DW_REG_ARM_WC3:      *regid_out = SYMS_REG_ARM_wc3;      break;
+    case DW_REG_ARM_WC4:      *regid_out = SYMS_REG_ARM_wc4;      break;
+    case DW_REG_ARM_WC5:      *regid_out = SYMS_REG_ARM_wc5;      break;
+    case DW_REG_ARM_WC6:      *regid_out = SYMS_REG_ARM_wc6;      break;
+    case DW_REG_ARM_WC7:      *regid_out = SYMS_REG_ARM_wc7;      break;
+    case DW_REG_ARM_D0:       *regid_out = SYMS_REG_ARM_d0;       break;
+    case DW_REG_ARM_D1:       *regid_out = SYMS_REG_ARM_d1;       break;
+    case DW_REG_ARM_D2:       *regid_out = SYMS_REG_ARM_d2;       break;
+    case DW_REG_ARM_D3:       *regid_out = SYMS_REG_ARM_d3;       break;
+    case DW_REG_ARM_D4:       *regid_out = SYMS_REG_ARM_d4;       break;
+    case DW_REG_ARM_D5:       *regid_out = SYMS_REG_ARM_d5;       break;
+    case DW_REG_ARM_D6:       *regid_out = SYMS_REG_ARM_d6;       break;
+    case DW_REG_ARM_D7:       *regid_out = SYMS_REG_ARM_d7;       break;
+    case DW_REG_ARM_D8:       *regid_out = SYMS_REG_ARM_d8;       break;
+    case DW_REG_ARM_D9:       *regid_out = SYMS_REG_ARM_d9;       break;
+    case DW_REG_ARM_D10:      *regid_out = SYMS_REG_ARM_d10;      break;
+    case DW_REG_ARM_D11:      *regid_out = SYMS_REG_ARM_d11;      break;
+    case DW_REG_ARM_D12:      *regid_out = SYMS_REG_ARM_d12;      break;
+    case DW_REG_ARM_D13:      *regid_out = SYMS_REG_ARM_d13;      break;
+    case DW_REG_ARM_D14:      *regid_out = SYMS_REG_ARM_d14;      break;
+    case DW_REG_ARM_D15:      *regid_out = SYMS_REG_ARM_d15;      break;
+    case DW_REG_ARM_D16:      *regid_out = SYMS_REG_ARM_d16;      break;
+    case DW_REG_ARM_D17:      *regid_out = SYMS_REG_ARM_d17;      break;
+    case DW_REG_ARM_D18:      *regid_out = SYMS_REG_ARM_d18;      break;
+    case DW_REG_ARM_D19:      *regid_out = SYMS_REG_ARM_d19;      break;
+    case DW_REG_ARM_D20:      *regid_out = SYMS_REG_ARM_d20;      break;
+    case DW_REG_ARM_D21:      *regid_out = SYMS_REG_ARM_d21;      break;
+    case DW_REG_ARM_D22:      *regid_out = SYMS_REG_ARM_d22;      break;
+    case DW_REG_ARM_D23:      *regid_out = SYMS_REG_ARM_d23;      break;
+    case DW_REG_ARM_D24:      *regid_out = SYMS_REG_ARM_d24;      break;
+    case DW_REG_ARM_D25:      *regid_out = SYMS_REG_ARM_d25;      break;
+    case DW_REG_ARM_D26:      *regid_out = SYMS_REG_ARM_d26;      break;
+    case DW_REG_ARM_D27:      *regid_out = SYMS_REG_ARM_d27;      break;
+    case DW_REG_ARM_D28:      *regid_out = SYMS_REG_ARM_d28;      break;
+    case DW_REG_ARM_D29:      *regid_out = SYMS_REG_ARM_d29;      break;
+    case DW_REG_ARM_D30:      *regid_out = SYMS_REG_ARM_d30;      break;
+    case DW_REG_ARM_D31:      *regid_out = SYMS_REG_ARM_d31;      break;
+    default: is_mapped = syms_false; break;
+    }
+#endif
+  } break;
+
+  default: is_mapped = syms_false; break;
+  }
+
+  return is_mapped;
+}
+
+SYMS_INTERNAL
+DW_REGWRITE_SIG(syms_regwrite_dwarf)
+{
+  SymsRegwrite *info = (SymsRegwrite *)context;
+  SymsRegID regid = SYMS_REG_null;
+  syms_uint write_size = 0;
+  if (syms_dw_regid_to_regid(arch, reg_index, &regid)) {
+    write_size = info->callback(info->context, arch, regid, value, syms_trunc_u32(value_size));
+  }
+  return (dw_uint)write_size;
+}
+
+SYMS_INTERNAL
+DW_REGREAD_SIG(syms_regread_dwarf)
+{
+  SymsRegread *info = (SymsRegread *)context;
+  SymsRegID regid = SYMS_REG_null;
+  syms_uint read_size = 0;
+  if (syms_dw_regid_to_regid(arch, reg_index, &regid)) {
+    read_size = info->callback(info->context, arch, regid, read_buffer, syms_trunc_u32(read_buffer_max));
+  }
+  return (dw_uint)read_size;
+}
+
+SYMS_INTERNAL
+DW_MEMREAD_SIG(syms_memread_dwarf)
+{
+  SymsMemread *memread = (SymsMemread *)context;
+  memread->result = memread->callback(memread->context, va, read_buffer, num_read);
+  return SYMS_RESULT_OK(memread->result) ? syms_true : syms_false;
+}
+
+SYMS_INTERNAL syms_bool
+syms_debug_file_iter_init_dwarf(SymsInstance *instance, SymsDebugFileIterDwarf *iter)
+{
+  SymsSecIter sec_iter;
+  if (syms_sec_iter_init(instance, &sec_iter)) {
+    SymsSection sec;
+    SymsSection debug_info;
+    SymsSection debug_abbrev;
+    SymsSection debug_str;
+    syms_memzero(&debug_info, sizeof(debug_info));
+    syms_memzero(&debug_abbrev, sizeof(debug_abbrev));
+    syms_memzero(&debug_str, sizeof(debug_str));
+    while (syms_sec_iter_next(&sec_iter, &sec)) {
+      if (syms_string_cmp_lit(sec.name, ".debug_info")) {
+        debug_info = sec;
+      } else if (syms_string_cmp_lit(sec.name, ".debug_abbrev")) {
+        debug_abbrev = sec;
+      } else if (syms_string_cmp_lit(sec.name, ".debug_str")) {
+        debug_str = sec;
+      }
+    }
+
+    DwInitdata init_data;
+    init_data.secs[DW_SEC_INFO].data_len = debug_info.data_size;
+    init_data.secs[DW_SEC_INFO].data = debug_info.data;
+    init_data.secs[DW_SEC_ABBREV].data_len = debug_abbrev.data_size;
+    init_data.secs[DW_SEC_ABBREV].data = debug_abbrev.data;
+    init_data.secs[DW_SEC_STR].data_len = debug_str.data_size;
+    init_data.secs[DW_SEC_STR].data = debug_str.data;
+
+    if (dw_init(&iter->context, syms_get_arch(instance), &init_data)) {
+      return dw_cu_iter_init(&iter->cu_iter, &iter->context);
+    }
+  }
+  return syms_false;
+}
+
+SYMS_INTERNAL syms_bool
+syms_debug_file_iter_next_dwarf(SymsDebugFileIterDwarf *iter, SymsString *dwo_path_out)
+{
+  DwCompileUnit cu;
+  syms_bool is_valid = syms_false;
+  while (dw_cu_iter_next(&iter->cu_iter, &cu)) {
+    *dwo_path_out = cu.dwo_name;
+    //*dwo_id_out = cu.dwo_id;
+    is_valid = !syms_string_is_null(*dwo_path_out);
+    if (is_valid) {
+      break;
+    }
+  }
+  return is_valid;
+}
diff --git a/Engine/Source/Runtime/SymsLib/src/dwarf/syms_dwarf_api.h b/Engine/Source/Runtime/SymsLib/src/dwarf/syms_dwarf_api.h
new file mode 100644
--- /dev/null
+++ b/Engine/Source/Runtime/SymsLib/src/dwarf/syms_dwarf_api.h
@@ -0,0 +1,143 @@
+// Copyright Epic Games, Inc. All Rights Reserved.
+#ifndef SYMS_DWARF_API_INCLUDE_H
+#define SYMS_DWARF_API_INCLUDE_H
+
+/******************************************************************************
+ * File   : syms_dwarf.h                                                      *
+ * Author : Nikita Smith                                                      *
+ * Created: 2020/06/20                                                        *
+ * Purpose: Declarations for dwarf.h interface                                *
+ ******************************************************************************/
+
+typedef struct SymsDebugInfoDwarf 
+{
+  DwContext context;
+} SymsDebugInfoDwarf;
+
+typedef struct 
+{
+   SymsString compile_dir;
+   SymsString dir;
+   SymsString file;
+} DwFilePath;
+
+typedef struct SymsModIterDwarf 
+{
+  DwCuIter impl;
+} SymsModIterDwarf;
+
+typedef struct SymsLineIterDwarf 
+{
+  DwLineIter impl;
+  DwDirIndex prev_file_index;
+  DwCompileUnit *cu;
+} SymsLineIterDwarf;
+
+typedef struct SymsMemberIterDwarf 
+{
+  DwMemberIter impl;
+} SymsMemberIterDwarf;
+
+typedef struct SymsGlobalIterDwarf 
+{
+  U32 reserved;
+} SymsGlobalIterDwarf;
+
+typedef struct SymsLocalDataIterDwarf 
+{
+  U32 reserved;
+} SymsLocalDataIterDwarf;
+
+typedef struct SymsFileIterDwarf 
+{
+  DwFileIter impl;
+} SymsFileIterDwarf;
+
+typedef struct SymsProcIterDwarf 
+{
+  DwProcIter impl;
+} SymsProcIterDwarf;
+
+typedef struct SymsArgIterDwarf 
+{
+  U32 reserved;
+} SymsArgIterDwarf;
+
+typedef struct SymsLocalIterDwarf 
+{
+  DwLocalIter impl;
+} SymsLocalIterDwarf;
+
+typedef struct SymsInlineIterDwarf 
+{
+  SymsAddr rva;
+  DwCompileUnit *cu;
+  DwAttribIter attribs;
+  U32 depth;
+} SymsInlineIterDwarf;
+
+typedef struct SymsConstIterDwarf
+{
+  U32 reserved;
+} SymsConstIterDwarf;
+
+typedef struct SymsTypeIterDwarf 
+{
+  U32 reserved;
+} SymsTypeIterDwarf;
+
+typedef struct SymsRangeDwarf
+{
+  syms_uint cu_index;
+  SymsAddr range_off;
+} SymsRangeDwarf;
+
+typedef struct
+{
+  DwContext context;
+  DwCuIter cu_iter;
+  DwAttribIter att_iter;
+} SymsDebugFileIterDwarf;
+
+SYMS_COMPILER_ASSERT(sizeof(SymsUWord) >= 8);
+SYMS_COMPILER_ASSERT(sizeof(SymsUWord) >= 4);
+SYMS_COMPILER_ASSERT(SYMS_GET_MEMBER_SIZE(DwContext, secs) == SYMS_GET_MEMBER_SIZE(DwInitdata, secs));
+SYMS_COMPILER_ASSERT(sizeof(SymsRangeDwarf) <= sizeof(SymsRangeImpl));
+SYMS_COMPILER_ASSERT(sizeof(DwTag) <= sizeof(SymsTypeIDImpl));
+SYMS_COMPILER_ASSERT(sizeof(DwType) <= sizeof(SymsTypeImpl));
+SYMS_COMPILER_ASSERT(sizeof(SymsDebugInfoDwarf) <= sizeof(SymsDebugInfoImpl));
+SYMS_COMPILER_ASSERT(sizeof(DwFilePath) <= sizeof(SymsStringRefImpl));
+SYMS_COMPILER_ASSERT(sizeof(SymsModIterDwarf) <= sizeof(SymsModIterImpl));
+SYMS_COMPILER_ASSERT(sizeof(DwCompileUnit) <= sizeof(SymsModImpl));
+SYMS_COMPILER_ASSERT(sizeof(SymsLineIterDwarf) <= sizeof(SymsLineIterImpl));
+SYMS_COMPILER_ASSERT(sizeof(SymsMemberIterDwarf) <= sizeof(SymsMemberIterImpl));
+SYMS_COMPILER_ASSERT(sizeof(SymsGlobalIterDwarf) <= sizeof(SymsGlobalIterImpl));
+SYMS_COMPILER_ASSERT(sizeof(SymsLocalDataIterDwarf) <= sizeof(SymsLocalDataIterImpl));
+SYMS_COMPILER_ASSERT(sizeof(SymsFileIterDwarf) <= sizeof(SymsFileIterImpl));
+SYMS_COMPILER_ASSERT(sizeof(SymsArgIterDwarf) <= sizeof(SymsArgIterImpl));
+SYMS_COMPILER_ASSERT(sizeof(SymsLocalIterDwarf) <= sizeof(SymsLocalIterImpl));
+SYMS_COMPILER_ASSERT(sizeof(SymsInlineIterDwarf) <= sizeof(SymsInlineIterImpl));
+SYMS_COMPILER_ASSERT(sizeof(SymsDebugFileIterDwarf) <= sizeof(SymsDebugFileIterImpl));
+typedef DW_MEMREAD_SIG(SymsDwMemreadSig);
+typedef DW_REGREAD_SIG(SymsDwRegreadSig);
+typedef DW_REGWRITE_SIG(SymsDwRegwriteSig);
+
+// -----------------------------------------------------------------------------
+
+SYMS_INTERNAL
+DW_REGREAD_SIG(syms_regread_dwarf);
+
+SYMS_INTERNAL
+DW_MEMREAD_SIG(syms_memread_dwarf);
+
+SYMS_INTERNAL syms_bool
+syms_dw_type_to_syms_type(DwType *dw_type, SymsType *syms_type);
+
+SYMS_INTERNAL syms_bool
+syms_proc_from_dw_proc(SymsDebugInfoDwarf *debug_info, DwCompileUnit *cu, DwProc *proc, SymsProc *proc_out);
+
+SYMS_INTERNAL SymsStringRef
+syms_string_ref_dw_path(SymsString compile_dir, SymsString dir, SymsString file);
+
+#endif /* SYMS_DWARF_API_INCLUDE_H */
+
diff --git a/Engine/Source/Runtime/SymsLib/src/dwarf/syms_dwarf_unwind.c b/Engine/Source/Runtime/SymsLib/src/dwarf/syms_dwarf_unwind.c
new file mode 100644
--- /dev/null
+++ b/Engine/Source/Runtime/SymsLib/src/dwarf/syms_dwarf_unwind.c
@@ -0,0 +1,1050 @@
+// Copyright Epic Games, Inc. All Rights Reserved.
+SYMS_INTERNAL const char *
+dw_get_reg_name(SymsArch arch, U32 regid)
+{
+  switch (arch) {
+  case SYMS_ARCH_X86: {
+#define X(arch, name, value) case value: return #name;
+    switch (regid) {
+      DW_REG_X86_LIST
+    }
+#undef X
+  } break;
+
+  case SYMS_ARCH_X64: {
+#define X(arch, name, value) case value: return #name;
+    switch (regid) {
+      DW_REG_X64_LIST
+    }
+#undef X
+  } break;
+
+  case SYMS_ARCH_ARM: {
+#define X(arch, name, value) case value: return #name;
+    switch (regid) {
+      DW_REG_ARM_LIST
+    }
+#undef X
+  } break;
+
+  default: SYMS_INVALID_CODE_PATH; break;
+  }
+
+  return "<uknown-regid>";
+}
+
+SYMS_INTERNAL syms_bool
+dw_is_cie_id(DwFrameInfo *frame_info, SymsSWord id)
+{
+  syms_bool result = syms_false;
+  switch (frame_info->source_type) {
+  case DW_VIRTUAL_UNWIND_DATA_DEBUG_FRAME: result = id == -1 || id == -1; break;
+  case DW_VIRTUAL_UNWIND_DATA_EH_FRAME:    result = (id == 0);  break;
+  case DW_VIRTUAL_UNWIND_DATA_NULL:        result = syms_false; break;
+  }
+  return result;
+}
+
+SYMS_INTERNAL dw_uint
+dw_regread(DwUserCallbacks *user_cbs, dw_uint regid, void *value, dw_uint value_size)
+{
+  dw_uint read_size = 0;
+
+  if (user_cbs->regread != NULL) {
+    read_size = user_cbs->regread(user_cbs->regread_ctx, user_cbs->arch_info.arch, regid, value, value_size);
+  } else {
+    SYMS_ASSERT_FAILURE("callback for reading registers is not specified");
+  }
+   
+  return read_size;
+}
+
+SYMS_INTERNAL dw_uint
+dw_regwrite(DwUserCallbacks *user_cbs, dw_uint regid, void *value, dw_uint value_size)
+{
+  dw_uint write_size = 0;
+
+  if (user_cbs->regwrite != NULL) {
+    write_size = user_cbs->regwrite(user_cbs->regwrite_ctx, user_cbs->arch_info.arch, regid, value, value_size);
+  } else {
+    SYMS_ASSERT_FAILURE("callback for writing registers is not specified");
+  }
+
+  return write_size;
+}
+
+SYMS_INTERNAL syms_bool
+dw_regread_uword(DwUserCallbacks *user_cbs, dw_uint regid, SymsUWord *value_out)
+{
+  dw_uint read_size;
+  syms_bool result;
+
+  SYMS_ASSERT(user_cbs->arch_info.word_size <= sizeof(*value_out));
+  *value_out = 0;
+  read_size = dw_regread(user_cbs, regid, value_out, user_cbs->arch_info.word_size);
+  result = (read_size == user_cbs->arch_info.word_size);
+
+  return result;
+}
+
+SYMS_INTERNAL syms_bool
+dw_regread_sword(DwUserCallbacks *user_cbs, dw_uint regid, SymsSWord *value_out)
+{
+  dw_uint read_size;
+  syms_bool result;
+
+  SYMS_ASSERT(user_cbs->arch_info.word_size <= sizeof(*value_out));
+  *value_out = 0;
+  read_size = dw_regread(user_cbs, regid, value_out, user_cbs->arch_info.word_size);
+  result = (read_size == user_cbs->arch_info.word_size);
+
+  return result;
+}
+
+SYMS_INTERNAL syms_bool
+dw_regwrite_uword(DwUserCallbacks *user_cbs, S64 regid, SymsUWord *value)
+{
+  dw_uint write_size;
+  syms_bool result = syms_false;
+  if (regid >= 0) {
+    SYMS_ASSERT(user_cbs->arch_info.word_size <= sizeof(*value));
+    write_size = dw_regwrite(user_cbs, syms_trunc_u32(regid), value, user_cbs->arch_info.word_size);
+    result = (write_size == user_cbs->arch_info.word_size);
+  }
+  return result;
+}
+
+SYMS_INTERNAL syms_bool
+dw_regwrite_sword(DwUserCallbacks *user_cbs, dw_uint regid, SymsSWord *value)
+{
+  dw_uint write_size = dw_regwrite(user_cbs, regid, value, user_cbs->arch_info.word_size);
+  syms_bool result = (write_size == user_cbs->arch_info.word_size);
+  return result;
+}
+
+SYMS_INTERNAL syms_bool
+dw_memread(DwUserCallbacks *user_cbs, SymsAddr va, void *buffer, dw_uint buffer_size)
+{
+  /* TODO(nick): replace bool with byte read counter */
+  syms_bool result = syms_false;
+
+  if (user_cbs->memread != NULL) {
+    result = user_cbs->memread(user_cbs->memread_ctx, va, buffer, buffer_size);
+  } else {
+    SYMS_ASSERT_FAILURE("callback for reading memory is not specified");
+  }
+
+  return result;
+}
+
+SYMS_INTERNAL syms_bool
+dw_parse_cie(DwFrameInfo *frame_info, DwBinRead *secdata, DwCommonInfoEntry *cie_out)
+{
+  dw_uint i;
+  SymsAddr augdata_end;
+
+  cie_out->init_cfi_offset = DW_INVALID_OFFSET;
+  cie_out->init_cfi_size = 0;
+  cie_out->end_offset = DW_INVALID_OFFSET;
+
+  cie_out->cie_id = 0;
+  cie_out->augmentation = syms_string_init(0,0);
+
+  cie_out->is_aug_sized = syms_false;
+  cie_out->is_sig_frame = syms_false;
+  cie_out->have_abi_maker = syms_false;
+
+  cie_out->lsda_encoding = DW_EH_PE_OMIT;
+  cie_out->fde_encoding = DW_EH_PE_OMIT;
+  cie_out->handler_encoding = DW_EH_PE_OMIT;
+
+  cie_out->addr_size = frame_info->arch_info.addr_size;
+  cie_out->segsel_size = 0;
+  cie_out->code_align = 0;
+  cie_out->data_align = 0;
+  cie_out->ret_addr_reg = 0;
+
+  switch (frame_info->arch_info.addr_size) {
+  case 4:     cie_out->fde_encoding = DW_EH_PE_UDATA4; break;
+  case 8:     cie_out->fde_encoding = DW_EH_PE_UDATA8; break;
+  default:    cie_out->fde_encoding = DW_EH_PE_OMIT;   break;
+  }
+
+  {
+    U64 cie_size;
+
+    cie_size = dw_bin_read_u32(secdata);
+    if (cie_size != 0xFFFFFFFFul) {
+      cie_out->end_offset = secdata->off + cie_size;
+      cie_out->cie_id = dw_bin_read_u32(secdata);
+    } else {
+      cie_size = dw_bin_read_u64(secdata);
+      cie_out->end_offset = secdata->off + cie_size;
+      cie_out->cie_id = dw_bin_read_u64(secdata);
+    }
+
+    switch (frame_info->source_type) {
+    case DW_VIRTUAL_UNWIND_DATA_DEBUG_FRAME: {
+      if (cie_out->cie_id != SYMS_UINT32_MAX || cie_out->cie_id != SYMS_UINT64_MAX) {
+        return syms_false;
+      }
+    } break;
+
+    case DW_VIRTUAL_UNWIND_DATA_EH_FRAME: {
+      if (cie_out->cie_id != 0) {
+        return syms_false;
+      }
+    } break;
+
+    default: SYMS_INVALID_CODE_PATH; break;
+    }
+  }
+
+  cie_out->version = dw_bin_read_u08(secdata);
+  if (cie_out->version != 1 && (cie_out->version < DW_CIE_VERSION || cie_out->version > DW_CIE_VERSION_MAX)) {
+    return syms_false;
+  }
+
+  cie_out->augmentation = dw_bin_read_string(secdata);
+#if 0
+  if (source_type == DW_STACK_UNWIND_DATA_SOURCE_DEBUG_FRAME) {
+    cie_out->addr_size = dw_bin_read_sleb128(secdata);
+    cie_out->segsel_size = dw_bin_read_sleb128(secdata);
+  } else {
+    cie_out->addr_size = 0;
+    cie_out->segsel_size = 0;
+  }
+#endif
+  cie_out->code_align = dw_bin_read_uleb128(secdata);
+  cie_out->data_align = dw_bin_read_sleb128(secdata);
+
+  if (cie_out->version == 1) {
+    cie_out->ret_addr_reg = dw_bin_read_u08(secdata);
+  } else {
+    cie_out->ret_addr_reg = dw_bin_read_uleb128(secdata);
+  }
+
+  augdata_end = DW_INVALID_OFFSET;
+  for (i = 0; i < cie_out->augmentation.len; ++i) {
+    switch (syms_string_peek_byte(cie_out->augmentation, i)) {
+    case 'z': {
+      U64 augdata_size = dw_bin_read_uleb128(secdata);
+      augdata_end = secdata->off + augdata_size;
+      cie_out->is_aug_sized = syms_true;
+    } break;
+
+    case 'L': {
+      cie_out->lsda_encoding = dw_bin_read_u08(secdata);
+      SYMS_ASSERT(cie_out->lsda_encoding);
+    } break;
+
+    case 'R': {
+      cie_out->fde_encoding = dw_bin_read_u08(secdata);
+      SYMS_ASSERT(cie_out->fde_encoding);
+    } break;
+
+    case 'S': {
+      cie_out->is_sig_frame = syms_true;
+      cie_out->have_abi_maker = syms_true;
+    } break;
+
+    case 'P': {
+      cie_out->handler_encoding = dw_bin_read_u08(secdata);
+      SYMS_ASSERT(cie_out->handler_encoding);
+      cie_out->handler_ip = dw_parse_pointer(frame_info, secdata, cie_out->handler_encoding);
+    } break;
+
+    default: {
+      if (!dw_bin_seek(secdata, augdata_end)) {
+        return syms_false;
+      }
+    } goto exit_augdata_read;
+    }
+  }
+  exit_augdata_read:;
+
+  SYMS_ASSERT(cie_out->end_offset >= secdata->off);
+  cie_out->init_cfi_offset = secdata->off;
+  cie_out->init_cfi_size = cie_out->end_offset - secdata->off;
+
+  return syms_true;
+}
+
+SYMS_INTERNAL SymsAddr
+dw_parse_pointer(DwFrameInfo *frame_info, DwBinRead *secdata, dw_uint encoding)
+{
+  SymsAddr init_off = secdata->off;
+  SymsAddr pointer = 0;
+
+  if (encoding == DW_EH_PE_OMIT) {
+    return 0;
+  } 
+  if ((encoding & DW_EH_PE_APPLY_MASK) == DW_EH_PE_ALIGNED) {
+    SymsAddr align;
+
+    SYMS_ASSERT(secdata->addr_size > 0);
+    align = (init_off + secdata->addr_size - 1) & (-secdata->addr_size);
+    if (!dw_bin_seek(secdata, align)) {
+      return 0;
+    }
+  }
+
+  if ((encoding & DW_EH_PE_FORMAT_MASK) == DW_EH_PE_PTR) {
+    encoding &= ~DW_EH_PE_PTR;
+    switch (secdata->mode) {
+    case DW_MODE_32BIT: encoding |= DW_EH_PE_SDATA4; break;
+    case DW_MODE_64BIT: encoding |= DW_EH_PE_SDATA8; break;
+    default: SYMS_INVALID_CODE_PATH;
+    }
+  }
+
+  switch (encoding & DW_EH_PE_FORMAT_MASK) {
+  case DW_EH_PE_UDATA2:   pointer = dw_bin_read_u16(secdata);     break;
+  case DW_EH_PE_UDATA4:   pointer = dw_bin_read_u32(secdata);     break;
+  case DW_EH_PE_UDATA8:   pointer = dw_bin_read_u64(secdata);     break;
+  case DW_EH_PE_ULEB128:  pointer = dw_bin_read_uleb128(secdata); break;
+
+  case DW_EH_PE_SDATA2:   pointer = (SymsAddr)dw_bin_read_s16(secdata);     break;
+  case DW_EH_PE_SDATA4:   pointer = (SymsAddr)dw_bin_read_s32(secdata);     break;
+  case DW_EH_PE_SDATA8:   pointer = (SymsAddr)dw_bin_read_s64(secdata);     break;
+  case DW_EH_PE_SLEB128:  pointer = (SymsAddr)dw_bin_read_sleb128(secdata); break;
+
+  default: SYMS_INVALID_CODE_PATH; break;
+  }
+
+  switch (encoding & DW_EH_PE_APPLY_MASK) {
+  case DW_EH_PE_ALIGNED: {
+    /* NOTE(nick): Ignore */
+  } break;
+
+  case DW_EH_PE_ABSPTR: {
+    /* NOTE(nick): leave pointer as is */
+  } break;
+
+  case DW_EH_PE_DATAREL: {
+    SYMS_ASSERT_FAILURE("DW_EH_PE_DATAREL not implemented");
+  } break;
+
+  case DW_EH_PE_FUNCREL: {
+    SYMS_ASSERT_FAILURE("DW_EH_PE_FUNCREL not implemented");
+  } break;
+
+  case DW_EH_PE_PCREL: {
+    SYMS_ASSERT(frame_info->eh_frame);
+    pointer += frame_info->eh_frame + init_off;
+  } break;
+
+  case DW_EH_PE_TEXTREL: {
+    // TODO(nick): IMPLEMENT
+    SYMS_ASSERT_FAILURE("DW_EH_PE_TEXTREL is not implemented");
+  }
+  default: SYMS_INVALID_CODE_PATH; break;
+  }
+
+  SYMS_ASSERT(pointer);
+  return pointer;
+}
+
+SYMS_INTERNAL syms_bool
+dw_cfi_program_init(SymsImageType img_type, DwCfiProgram *program)
+{
+  dw_uint i;
+  syms_bool is_inited;
+
+  switch (img_type) {
+  case SYMS_IMAGE_NULL: program->reg_count = 0;                             break;
+  case SYMS_IMAGE_NT:   program->reg_count = DW_PRESERVED_REGS_COUNT_NT;    break;
+  case SYMS_IMAGE_ELF:  program->reg_count = DW_PRESERVED_REGS_COUNT_UNIX;  break;
+  }
+
+  for (i = 0; i < SYMS_ARRAY_SIZE(program->rules.value); ++i) {
+    program->rules.type[i] = DW_CFI_REGISTER_TYPE_SAME;
+    program->rules.value[i].w = 0;
+  }
+
+  program->frame           = &program->stack[DW_CFI_PROGRAM_STACK_MAX - 1];
+  program->setup_cfa       = syms_true;
+  program->ret_addr_regid  = SYMS_UINT32_MAX;
+  program->cfa             = 0;
+
+  is_inited = (program->reg_count > 0);
+  return is_inited;
+}
+
+SYMS_INTERNAL syms_bool
+dw_compile_cfi_table_row(DwCfiProgram *program, DwFrameInfo *frame_info, DwCommonInfoEntry *cie, DwFrameDescEntry *fde, DwBinRead *cfi_data, SymsAddr ip)
+{
+  SymsAddr decode_ip;
+  DwCfiRow *rules_frame = &program->stack[DW_CFI_PROGRAM_STACK_MAX - 1];
+   
+  if (fde) {
+    decode_ip = fde->start_ip;
+    if (ip > frame_info->image_base) {
+      decode_ip += frame_info->image_base;
+    }
+  } else {
+    decode_ip = 0;
+  }
+
+  while (cfi_data->off < cfi_data->max) {
+    U8 opcode = dw_bin_read_u08(cfi_data);
+    U8 operand = 0;
+
+    if ((opcode & DW_CFA_OPCODE_MASK) != 0) {
+      operand = opcode & DW_CFA_OPER_MASK;
+      opcode = opcode & DW_CFA_OPCODE_MASK;
+    }
+
+    if (decode_ip >= ip) {
+      /* NOTE(nick): Instruction Pointer might be at any point in procedure body:
+       * prologue, body, or epilogue. And instructions that manage
+       * stack pointer might not be executed yet and we have to take this into account
+       * by comparing virtual-instruction-pointer 'decode_ip' with instruction pointer 'ip',
+       * which is actual location in procedures body. */
+      break;
+    }
+
+    switch (opcode) {
+    case DW_CFA_NOP: {
+    } break;
+
+    case DW_CFA_SAME_VALUE: {
+      SymsUWord reg = dw_bin_read_uleb128(cfi_data);
+      if (reg < program->reg_count) {
+        program->rules.type[reg] = DW_CFI_REGISTER_TYPE_SAME;
+        program->rules.value[reg].w = 0;
+      }
+    } break;
+
+    case DW_CFA_UNDEFINED: {
+      SymsUWord reg = dw_bin_read_uleb128(cfi_data);
+      if (reg < program->reg_count) {
+        program->rules.type[reg] = DW_CFI_REGISTER_TYPE_UNDEF;
+        program->rules.value[reg].w = 0;
+      }
+    } break;
+
+    case DW_CFA_ADVANCE_LOC: {
+      decode_ip += operand * cie->code_align;
+    } break;
+
+    case DW_CFA_ADVANCE_LOC1: {
+      U8 value = dw_bin_read_u08(cfi_data);
+      decode_ip += value * cie->code_align;
+    } break;
+
+    case DW_CFA_ADVANCE_LOC2: {
+      U16 value = dw_bin_read_u16(cfi_data);
+      decode_ip += value * cie->code_align;
+    } break;
+
+    case DW_CFA_ADVANCE_LOC4: {
+      dw_uint value = dw_bin_read_u32(cfi_data);
+      decode_ip += value * cie->code_align;
+    } break;
+
+    case DW_CFA_OFFSET: {
+      DwOffset reg = operand;
+      if (reg < program->reg_count) {
+        program->rules.type[reg] = DW_CFI_REGISTER_TYPE_CFAREL;
+        program->rules.value[reg].w = (S64)dw_bin_read_uleb128(cfi_data) * cie->data_align;
+      } else {
+        // TODO(nick): handle invalid register
+      }
+    } break;
+
+    case DW_CFA_VAL_OFFSET: {
+      DwOffset reg = dw_bin_read_sleb128(cfi_data);
+      DwOffset value = dw_bin_read_sleb128(cfi_data);
+
+      if ((dw_uint)reg < program->reg_count) {
+        program->rules.type[reg] = DW_CFI_REGISTER_TYPE_CFAREL;
+        program->rules.value[reg].w = value * cie->data_align;
+      } else {
+        // TODO(nick): handle invalid register
+      }
+    } break;
+
+    case DW_CFA_SET_LOC: {
+      decode_ip = dw_parse_pointer(frame_info, cfi_data, cie->fde_encoding);
+    } break;
+
+    case DW_CFA_DEF_CFA_OFFSET: {
+      program->rules.type[DW_CFA_COLUMN_OFF] = DW_CFI_REGISTER_TYPE_UNDEF;
+      program->rules.value[DW_CFA_COLUMN_OFF].w = (S64)dw_bin_read_uleb128(cfi_data);
+    } break;
+
+    case DW_CFA_DEF_CFA_REGISTER: {
+      SymsUWord reg = dw_bin_read_uleb128(cfi_data);
+
+      if (reg < program->reg_count) {
+        program->rules.type[DW_CFA_COLUMN_REG] = DW_CFI_REGISTER_TYPE_REG;
+        program->rules.value[DW_CFA_COLUMN_REG].w = (SymsSWord)reg;
+      } else {
+        // TODO(nick): handle invalid register
+      }
+    } break;
+
+    case DW_CFA_EXPR: {
+      dw_uint reg = dw_bin_read_uleb128(cfi_data);
+      dw_uint expr_size = dw_bin_read_uleb128(cfi_data);
+
+      if (expr_size <= SYMS_UINT32_MAX) {
+        if (reg < program->reg_count) {
+          if (dw_bin_skip(cfi_data, expr_size)) {
+            program->rules.type[reg] = DW_CFI_REGISTER_TYPE_EXPR;
+            program->rules.value[reg].e.ops = dw_bin_at(cfi_data);
+            program->rules.value[reg].e.ops_size = expr_size;
+            program->rules.value[reg].e.frame_base = 0;
+            program->rules.value[reg].e.member_location = 0;
+            program->rules.value[reg].e.cfa = program->cfa;
+          } else {
+            // TODO(nick): handle invalid expression size
+          }
+        } else {
+          // TODO(nick): handle invalid register
+        }
+      } else {
+        program->rules.type[reg] = DW_CFI_REGISTER_TYPE_INVALID;
+        program->rules.value[reg].e.ops = 0;
+        program->rules.value[reg].e.ops_size = 0;
+        program->rules.value[reg].e.frame_base = 0;
+        program->rules.value[reg].e.member_location = 0;
+        program->rules.value[reg].e.cfa = 0;
+      }
+    } break;
+
+    case DW_CFA_DEF_CFA: {
+      /* NOTE(nick): Defining Canonical Frame Address as 'register + offset' */
+
+      SymsUWord reg = dw_bin_read_uleb128(cfi_data);
+      SymsUWord off = dw_bin_read_uleb128(cfi_data);
+
+      if (reg < program->reg_count) {
+        program->rules.type[DW_CFA_COLUMN_REG] = DW_CFI_REGISTER_TYPE_REG;
+        program->rules.value[DW_CFA_COLUMN_REG].w = (SymsSWord)reg;
+
+        program->rules.type[DW_CFA_COLUMN_OFF] = DW_CFI_REGISTER_TYPE_UNDEF;
+        program->rules.value[DW_CFA_COLUMN_OFF].w = (SymsSWord)off;
+      } else {
+        // TODO(nick): handle invalid register
+      }
+    } break;
+
+    case DW_CFA_DEF_CFA_EXPR: {
+      /* NOTE(nick): Defining Canonical Frame Address as a dwarf expression */
+
+      U64 expr_size = dw_bin_read_uleb128(cfi_data);
+      if (expr_size < SYMS_UINT32_MAX) {
+        void *expr = dw_bin_at(cfi_data);
+        if (dw_bin_skip(cfi_data, expr_size)) {
+          program->rules.type[DW_CFA_COLUMN_REG] = DW_CFI_REGISTER_TYPE_EXPR;
+          program->rules.value[DW_CFA_COLUMN_REG].e.ops = expr;
+          program->rules.value[DW_CFA_COLUMN_REG].e.ops_size = (dw_uint)expr_size;
+          program->rules.value[DW_CFA_COLUMN_REG].e.frame_base = 0;
+          program->rules.value[DW_CFA_COLUMN_REG].e.member_location = 0;
+          program->rules.value[DW_CFA_COLUMN_REG].e.cfa = program->cfa;
+        }
+      } else {
+        program->rules.type[DW_CFA_COLUMN_REG] = DW_CFI_REGISTER_TYPE_INVALID;
+        program->rules.value[DW_CFA_COLUMN_REG].e.ops = 0;
+        program->rules.value[DW_CFA_COLUMN_REG].e.ops_size = 0;
+        program->rules.value[DW_CFA_COLUMN_REG].e.frame_base = 0;
+        program->rules.value[DW_CFA_COLUMN_REG].e.member_location = 0;
+        program->rules.value[DW_CFA_COLUMN_REG].e.cfa = 0;
+      }
+    } break;
+
+    case DW_CFA_REGISTER: {
+      SymsUWord reg = dw_bin_read_uleb128(cfi_data);
+      SymsUWord value = dw_bin_read_uleb128(cfi_data);
+
+      if (reg < program->reg_count) {
+        program->rules.type[reg] = DW_CFI_REGISTER_TYPE_REG;
+        program->rules.value[reg].w = (SymsSWord)value;
+      } else {
+        // TODO(nick): handle invalid register
+      }
+    } break;
+
+    case DW_CFA_REMEMBER_STATE: {
+      if (rules_frame < &program->stack[0]) {
+        SYMS_ASSERT_FAILURE("stack overflow");
+        return syms_false;
+      }
+
+      *rules_frame = program->rules;
+      --rules_frame;
+    } break;
+
+    case DW_CFA_RESTORE_STATE: {
+      if (rules_frame >= &program->stack[DW_CFI_PROGRAM_STACK_MAX]) {
+        SYMS_ASSERT_FAILURE("stack underflow");
+        return syms_false;
+      }
+
+      ++rules_frame;
+      program->rules = *rules_frame;
+    } break;
+
+    case DW_CFA_VAL_EXPR: {
+      SYMS_ASSERT_FAILURE("DW_CFA_VAL_EXPR");
+    } break;
+
+    case DW_CFA_RESTORE: {
+      SYMS_ASSERT_FAILURE("DW_CFA_RESTORE");
+    } break;
+
+    case DW_CFA_OFFSET_EXT: {
+      SYMS_ASSERT_FAILURE("DW_CFA_OFFSET_EXT");
+    } break;
+
+    case DW_CFA_RESTORE_EXT: {
+      SYMS_ASSERT_FAILURE("DW_CFA_RESTORE_EXT");
+    } break;
+
+    case DW_CFA_OFFSET_EXT_SF: {
+      SYMS_ASSERT_FAILURE("DW_CFA_OFFSET_EXT_SF");
+    } break;
+
+    case DW_CFA_DEF_CFA_SF: {
+      SYMS_ASSERT_FAILURE("DW_CFA_DEF_CFA_SF");
+    } break;
+
+    case DW_CFA_DEF_CFA_OFFSET_SF: {
+      SYMS_ASSERT_FAILURE("DW_CFA_DEF_CFA_OFFSET_SF");
+    } break;
+
+    default: {
+      if (opcode >= DW_CFA_USER_LO && opcode <= DW_CFA_USER_HI) {
+        // NOTE(nick): User implemented CFA opcode
+      } else {
+        SYMS_INVALID_CODE_PATH;
+      }
+    } break;
+    }
+  }
+
+  return syms_true;
+}
+
+SYMS_INTERNAL syms_bool
+dw_apply_cfi_table_row(DwCfiProgram *program, DwFrameInfo *frame, DwUserCallbacks *user_cbs)
+{
+  SymsUWord cfa = 0;
+  dw_uint regid;
+  
+  if (program->rules.type[DW_CFA_COLUMN_REG] == DW_CFI_REGISTER_TYPE_REG) {
+    SYMS_ASSERT(sizeof(cfa) >= frame->arch_info.word_size);
+    if (program->rules.value[DW_CFA_COLUMN_REG].w == DW_REG_X64_RSP && 
+        program->reg_count > DW_REG_X64_RSP && 
+        program->rules.type[DW_REG_X64_RSP] == DW_CFI_REGISTER_TYPE_SAME) {
+      cfa = program->cfa;
+    } else {
+      if (program->rules.value[DW_CFA_COLUMN_REG].w < 0) {
+        return syms_false;
+      }
+      regid = (dw_uint)program->rules.value[DW_CFA_COLUMN_REG].w;
+      if (!dw_regread_uword(user_cbs, regid, &cfa)) {
+        return syms_false;
+      }
+    }
+    if (program->rules.value[DW_CFA_COLUMN_OFF].w < 0 && program->rules.value[DW_CFA_COLUMN_OFF].w < (SymsSWord)cfa) {
+      return syms_false;
+    }
+    cfa += (dw_uint)program->rules.value[DW_CFA_COLUMN_OFF].w;
+  } else {
+    DwLocation loc;
+    syms_bool is_decoded;
+
+    SYMS_ASSERT(program->rules.type[DW_CFA_COLUMN_REG] == DW_CFI_REGISTER_TYPE_EXPR);
+    is_decoded = dw_decode_location_expr(
+        &program->rules.value[DW_CFA_COLUMN_REG].e, 
+        frame->arch_info.mode, frame->arch_info.arch,
+        user_cbs->memread_ctx, user_cbs->memread,
+        user_cbs->regread_ctx, user_cbs->regread,
+        &loc);
+
+    if (is_decoded && loc.type == DW_LOCATION_ADDR) {
+      cfa = loc.u.addr;
+    } else {
+      return syms_false;
+    }
+  }
+
+  for (regid = 0; regid < program->reg_count; ++regid) {
+    DwCfiRegisterType rule_type = program->rules.type[regid];
+    DwCfiRegValue rule_value = program->rules.value[regid];
+
+    switch (rule_type) {
+    case DW_CFI_REGISTER_TYPE_UNDEF: {
+      /* NOTE(nick): Undefined register value */
+
+      if (program->ret_addr_regid == regid) {
+        SymsUWord dummy_value = 0;
+        if (!dw_regwrite_uword(user_cbs, program->ret_addr_regid, &dummy_value)) {
+          return syms_false;
+        }
+      }
+    } break;
+
+    case DW_CFI_REGISTER_TYPE_SAME: {
+      /* NOTE(nick): Register value is not changed */
+    } break;
+
+    case DW_CFI_REGISTER_TYPE_CFAREL: {
+      SymsAddr value_va;
+      SymsUWord value;
+
+      SymsSWord temp = (SymsSWord)cfa + rule_value.w;
+      if (temp < 0) {
+        return syms_false;
+      }
+      value_va = (SymsAddr)temp;
+      value = 0;
+      if (!dw_memread(user_cbs, value_va, &value, frame->arch_info.word_size)) {
+        return syms_false;
+      }
+
+      if (!dw_regwrite_uword(user_cbs, regid, &value)) {
+        return syms_false;
+      }
+    } break;
+
+    case DW_CFI_REGISTER_TYPE_REG: {
+      SymsSWord value;
+
+      if (rule_value.w < 0) {
+        return syms_false;
+      }
+      value = 0;
+      if (!dw_regread_sword(user_cbs, (dw_uint)rule_value.w, &value)) {
+        return syms_false;
+      }
+      value += program->rules.value[rule_value.w].w;
+      if (!dw_regwrite_sword(user_cbs, regid, &value)) {
+        return syms_false;
+      }
+    } break;
+
+    case DW_CFI_REGISTER_TYPE_EXPR: {
+      SYMS_ASSERT_FAILURE("IMPLEMENT::CFI_REGISTER_TYPE_EXPR");
+    } break;
+
+    case DW_CFI_REGISTER_TYPE_VAL_EXPR: {
+      SYMS_ASSERT_FAILURE("IMPLEMENT::CFI_REGISTER_TYPE_VAL_EXPR");
+    } break;
+
+    default: SYMS_ASSERT_FAILURE("IMPLEMENT::INVALID_DEFAULT_CASE");
+    }
+  }
+
+  program->cfa = cfa;
+
+  return syms_true;
+}
+
+SYMS_INTERNAL syms_bool
+dw_fde_iter_init(DwFrameInfo *frame_info, DwBinRead frame_data, DwFrameDescEntryIter *it_out)
+{
+  syms_bool result = syms_false;
+
+  it_out->frame_info = frame_info;
+  it_out->secdata = frame_data;
+  it_out->cie_offset = 0;
+
+  if (dw_bin_seek(&it_out->secdata, it_out->cie_offset)) {
+    if (dw_parse_cie(frame_info, &it_out->secdata, &it_out->cie)) {
+      if (dw_bin_seek(&it_out->secdata, it_out->cie.end_offset)) {
+        result = syms_true;
+      }
+    }
+  }
+
+  return result;
+}
+
+SYMS_INTERNAL syms_bool
+dw_fde_iter_next(DwFrameDescEntryIter *it, DwFrameDescEntry *fde_out)
+{
+  DwBinRead *secdata = &it->secdata;
+  SymsAddr fde_end_offset;
+  SymsAddr cie_offset;
+  U64 fde_size;
+  SymsAddr base_offset;
+  DwCommonInfoEntry *cie;
+
+reparse:;
+
+  fde_end_offset = 0;
+  cie_offset = 0;
+  base_offset = secdata->off;
+  cie = NULL;
+
+  fde_size = dw_bin_read_u32(secdata);
+  if (fde_size != 0xFFFFFFFF) {
+    S32 cie_id;
+
+    if (fde_size == 0) {
+      return syms_false;
+    }
+
+    fde_end_offset = secdata->off + fde_size;
+
+    cie_id = dw_bin_read_s32(secdata);
+
+    if (dw_is_cie_id(it->frame_info, cie_id)) {
+      SYMS_ASSERT(fde_size);
+
+      if (!dw_bin_seek(secdata, base_offset + fde_size + 4)) {
+        // TODO(nick): handle corrupted data
+        return syms_false;
+      }
+
+      goto reparse;
+    }
+    if (cie_id < 0) {
+      return syms_false;
+    }
+    cie_offset = (SymsAddr)cie_id;
+
+    switch (it->frame_info->source_type) {
+    case DW_VIRTUAL_UNWIND_DATA_DEBUG_FRAME: {
+      cie_offset += base_offset;
+    } break;
+
+    case DW_VIRTUAL_UNWIND_DATA_EH_FRAME: {
+      if (cie_offset <= base_offset + 4) {
+        return syms_false;
+      }
+      cie_offset = (base_offset + 4) - cie_offset;
+    } break;
+
+    default: break;
+    }
+  } else {
+    SYMS_ASSERT_FAILURE("64bit version not implemented yet");
+  }
+
+  fde_out->data_off = base_offset;
+  fde_out->data_size = fde_size;
+
+  if (it->cie_offset != cie_offset ) {
+    SymsAddr pushed_offset = secdata->off;
+    dw_bin_seek(secdata, cie_offset);
+    if (!dw_parse_cie(it->frame_info, secdata, &it->cie)) {
+      SYMS_INVALID_CODE_PATH; // TODO(nick): error handling
+    }
+    dw_bin_seek(secdata, pushed_offset);
+
+    it->cie_offset = cie_offset;
+  }
+  cie = &it->cie;
+
+  fde_out->start_ip = 0;
+  fde_out->range_ip = 0;
+  fde_out->lsda_ip = 0;
+
+  switch (it->frame_info->source_type) {
+  case DW_VIRTUAL_UNWIND_DATA_EH_FRAME: {
+    fde_out->start_ip = dw_parse_pointer(it->frame_info, secdata, cie->fde_encoding);
+    fde_out->range_ip = dw_parse_pointer(it->frame_info, secdata, cie->fde_encoding & DW_EH_PE_FORMAT_MASK);
+
+    if (cie->is_aug_sized) {
+      SymsAddr aug_size = dw_bin_read_uleb128(secdata);
+      SymsAddr aug_end_addr = secdata->off + aug_size;
+
+      if (!dw_bin_skip(secdata, aug_size)) {
+        return syms_false;
+      }
+
+      if (aug_end_addr != secdata->off) {
+        // TODO(nick): handle corrupted data
+        return syms_false;
+      }
+    }
+  } break;
+
+  case DW_VIRTUAL_UNWIND_DATA_DEBUG_FRAME: {
+    fde_out->start_ip = dw_bin_read_addr(secdata);
+    fde_out->range_ip = dw_bin_read_addr(secdata);
+  } break;
+
+  default: SYMS_INVALID_CODE_PATH;
+  }
+  if (fde_out->start_ip == 0 || fde_out->range_ip == 0) {
+    return syms_false;
+  }
+
+  if (cie->have_abi_maker) {
+    U16 abi = dw_bin_read_u16(secdata);
+    U16 tag = dw_bin_read_u16(secdata);
+    (void)abi; (void)tag;
+  }
+
+  fde_out->cfi_offset = secdata->off;
+  fde_out->cfi_size = fde_end_offset - secdata->off;
+
+  if (!dw_bin_seek(secdata, fde_end_offset)) {
+    return syms_false;
+  }
+
+  return syms_true;
+}
+
+SYMS_INTERNAL syms_bool
+dw_find_fde(DwFrameInfo *frame_info, DwBinRead frame_data, SymsAddr ip, DwCommonInfoEntry *cie_out, DwFrameDescEntry *fde_out)
+{
+  DwFrameDescEntryIter it;
+  syms_bool result = syms_false;
+
+  // TODO(nick): Replace this linear search with binary!
+  if (dw_fde_iter_init(frame_info, frame_data, &it)) {
+    DwFrameDescEntry fde;
+
+    if (ip >= frame_info->image_base)
+      ip -= frame_info->image_base;
+
+    while (dw_fde_iter_next(&it, &fde)) {
+      if (ip >= fde.start_ip && ip < (fde.start_ip + fde.range_ip)) {
+        *cie_out = it.cie;
+        *fde_out = fde;
+
+        result = syms_true;
+        break;
+      }
+    }
+  }
+
+  return result;
+}
+
+DW_API syms_bool
+dw_virtual_unwind_init(SymsImageType image_type, DwVirtualUnwind *context_out)
+{
+  return dw_cfi_program_init(image_type, &context_out->program);
+}
+
+DW_API syms_bool
+dw_virtual_unwind_frame(DwVirtualUnwind *context,
+            SymsArch arch,
+            DwVirtualUnwindDataType source_type,
+            void *sec_bytes, SymsUMM sec_bytes_size,
+            SymsAddr image_base,
+            SymsAddr sec_bytes_base,
+            void *memread_ctx, dw_memread_sig *memread,
+            void *regread_ctx, dw_regread_sig *regread,
+            void *regwrite_ctx, dw_regwrite_sig *regwrite)
+{
+  SymsUWord ip;
+  SymsUWord sp;
+  U32 i;
+
+  DwUserCallbacks user_cbs;
+  DwCommonInfoEntry cie;
+  DwFrameDescEntry fde;
+  DwBinRead cfi_data;
+  DwBinRead frame_data;
+  DwFrameInfo frame;
+
+  frame.source_type = source_type;
+  frame.eh_frame = sec_bytes_base;
+  frame.image_base = image_base;
+  switch (arch) {
+  case SYMS_ARCH_X86: {
+    frame.arch_info.mode      = DW_MODE_64BIT;
+    frame.arch_info.arch      = arch;
+    frame.arch_info.ip_regid  = DW_REG_X86_EIP; 
+    frame.arch_info.sp_regid  = DW_REG_X86_ESP; 
+    frame.arch_info.word_size = 4;
+    frame.arch_info.word_size = 4;
+  } break;
+
+  case SYMS_ARCH_X64: {
+    frame.arch_info.mode      = DW_MODE_32BIT;
+    frame.arch_info.arch      = arch;
+    frame.arch_info.ip_regid  = DW_REG_X64_RIP; 
+    frame.arch_info.sp_regid  = DW_REG_X64_RSP; 
+    frame.arch_info.word_size = 8;
+    frame.arch_info.addr_size = 8;
+  } break;
+
+  default: {
+    frame.arch_info.mode      = DW_MODE_NULL;
+    frame.arch_info.arch      = SYMS_ARCH_NULL;
+    frame.arch_info.ip_regid  = SYMS_UINT32_MAX;
+    frame.arch_info.sp_regid  = SYMS_UINT32_MAX;
+    frame.arch_info.word_size = 0;
+  } break;
+  }
+
+  user_cbs.arch_info    = frame.arch_info;
+  user_cbs.memread_ctx  = memread_ctx;
+  user_cbs.memread      = memread;
+  user_cbs.regread_ctx  = regread_ctx;
+  user_cbs.regread      = regread;
+  user_cbs.regwrite_ctx = regwrite_ctx;
+  user_cbs.regwrite     = regwrite;
+  
+  ip = 0;
+  if (!dw_regread_uword(&user_cbs, frame.arch_info.ip_regid, &ip)) {
+    return syms_false;
+  }
+
+  sp = 0;
+  if (!dw_regread_uword(&user_cbs, frame.arch_info.sp_regid, &sp)) {
+    return syms_false;
+  }
+
+  frame_data = dw_bin_read_init(DW_MODE_32BIT, frame.arch_info.word_size, sec_bytes, sec_bytes_size);
+  if (dw_bin_peek_u32(&frame_data) == 0xFFFFFFFF) {
+    frame_data = dw_bin_read_init(DW_MODE_64BIT, frame.arch_info.word_size, sec_bytes, sec_bytes_size);
+  }
+
+  if (!dw_find_fde(&frame, frame_data, ip, &cie, &fde)) {
+    return syms_false;
+  }
+
+  for (i = 0; i < SYMS_ARRAY_SIZE(context->program.rules.value); ++i) {
+    context->program.rules.type[i] = DW_CFI_REGISTER_TYPE_SAME;
+    context->program.rules.value[i].w = 0;
+  }
+
+  /* NOTE(nick): According to DWARF5.pdf instructions from CIE must be
+   * executed first to init "state", and should be stored away for later 
+   * use by opcodes like DW_CFA_REMEMBER_STATE and DW_CFA_RESTORE_STATE. */
+  cfi_data = dw_bin_read_init(frame.arch_info.mode, frame.arch_info.addr_size, sec_bytes, cie.init_cfi_offset + cie.init_cfi_size);
+  if (!dw_bin_seek(&cfi_data, cie.init_cfi_offset)) {
+    // TODO(nick): handle corrupted data
+    return syms_false;
+  }
+  if (!dw_compile_cfi_table_row(&context->program, &frame, &cie, NULL, &cfi_data, (SymsAddr)~0)) {
+    return syms_false;
+  }
+  context->program.ret_addr_regid = cie.ret_addr_reg;
+
+
+  if (context->program.setup_cfa) {
+    context->program.setup_cfa = syms_false;
+    context->program.cfa = sp;
+  }
+
+  /* NOTE(nick): Parse FDE instructions, these tell us how to unwind registers. */
+  cfi_data = dw_bin_read_init(frame.arch_info.mode, frame.arch_info.addr_size, sec_bytes, fde.cfi_offset + fde.cfi_size);
+  if (!dw_bin_seek(&cfi_data, fde.cfi_offset)) {
+    SYMS_ASSERT_FAILURE("IMPLEMENT::ERROR");
+    return syms_false;
+  }
+  if (!dw_compile_cfi_table_row(&context->program, &frame, &cie, &fde, &cfi_data, ip)) {
+    SYMS_ASSERT_FAILURE("IMPLEMENT::ERROR");
+    return syms_false;
+  }
+
+  /* NOTE(nick): Applying rules for unwinding to registers. */
+  if (!dw_apply_cfi_table_row(&context->program, &frame, &user_cbs)) {
+    /* SYMS_ASSERT_FAILURE("IMPLEMENT::ERROR"); */
+    return syms_false;
+  }
+
+  return syms_true;
+}
diff --git a/Engine/Source/Runtime/SymsLib/src/dwarf/syms_dwarf_unwind.h b/Engine/Source/Runtime/SymsLib/src/dwarf/syms_dwarf_unwind.h
new file mode 100644
--- /dev/null
+++ b/Engine/Source/Runtime/SymsLib/src/dwarf/syms_dwarf_unwind.h
@@ -0,0 +1,253 @@
+// Copyright Epic Games, Inc. All Rights Reserved.
+#ifndef SYMS_DWARF_UNWIND_INCLUDE_H
+#define SYMS_DWARF_UNWIND_INCLUDE_H
+
+#ifndef DW_API
+#define DW_API SYMS_API
+#endif
+
+typedef struct DwArchInfo 
+{
+  DwMode mode;
+  SymsArch arch;
+  dw_uint word_size;
+  dw_uint addr_size;
+  dw_uint ip_regid;
+  dw_uint sp_regid;
+} DwArchInfo; 
+
+typedef struct DwUserCallbacks 
+{
+  DwArchInfo arch_info;
+
+  void *regread_ctx;
+  dw_regread_sig *regread;
+
+  void *regwrite_ctx;
+  dw_regwrite_sig *regwrite;
+
+  void *memread_ctx;
+  dw_memread_sig *memread;
+} DwUserCallbacks;
+
+#define DW_PRESERVED_REGS_COUNT_NT      33
+#define DW_PRESERVED_REGS_COUNT_UNIX    17
+#define DW_PRESERVED_REGS_MAX           (SYMS_MAX(DW_PRESERVED_REGS_COUNT_NT, DW_PRESERVED_REGS_COUNT_UNIX))
+
+#define DW_CFA_OPCODE_MASK 0xC0
+#define DW_CFA_OPER_MASK   0x3F
+typedef enum 
+{
+  DW_CFA_NOP                   = 0x0,
+  DW_CFA_SET_LOC               = 0x1,
+  DW_CFA_ADVANCE_LOC1          = 0x2,
+  DW_CFA_ADVANCE_LOC2          = 0x3,
+  DW_CFA_ADVANCE_LOC4          = 0x4,
+  DW_CFA_OFFSET_EXT            = 0x5,
+  DW_CFA_RESTORE_EXT           = 0x6,
+  DW_CFA_UNDEFINED             = 0x07,
+  DW_CFA_SAME_VALUE            = 0x8,
+  DW_CFA_REGISTER              = 0x9,
+  DW_CFA_REMEMBER_STATE        = 0xA,
+  DW_CFA_RESTORE_STATE         = 0xB,
+  DW_CFA_DEF_CFA               = 0xC,
+  DW_CFA_DEF_CFA_REGISTER      = 0xD,
+  DW_CFA_DEF_CFA_OFFSET        = 0xE,
+  DW_CFA_DEF_CFA_EXPR          = 0xF,
+  DW_CFA_EXPR                  = 0x10,
+  DW_CFA_OFFSET_EXT_SF         = 0x11,
+  DW_CFA_DEF_CFA_SF            = 0x12,
+  DW_CFA_DEF_CFA_OFFSET_SF     = 0x13,
+  DW_CFA_VAL_OFFSET            = 0x14,
+  DW_CFA_VAL_OFFSET_SF         = 0x15,
+  DW_CFA_VAL_EXPR              = 0x16,
+
+  DW_CFA_ADVANCE_LOC = 0x40,
+  DW_CFA_OFFSET      = 0x80,
+  DW_CFA_RESTORE     = 0xC0,
+
+  DW_CFA_USER_LO = 0x1C,
+  DW_CFA_USER_HI = 0x3F
+} DwCFAOpcodeType;
+
+#define DW_EH_PE_FORMAT_MASK    0x0F
+#define DW_EH_PE_APPLY_MASK     0x70
+#define DW_EH_PE_INDIRECT       0x80
+
+#define DW_EH_PE_OMIT       0xFFu
+/* Pointer sized unsigned value */
+#define DW_EH_PE_PTR        0x00u
+/* Unsigned LE base-128 value */
+#define DW_EH_PE_ULEB128    0x01u
+/* Unsigned 16-bit value */
+#define DW_EH_PE_UDATA2     0x02u
+/* Unsigned 32-bit value */
+#define DW_EH_PE_UDATA4     0x03u
+/* Unsigned 64-bit value */
+#define DW_EH_PE_UDATA8     0x04u
+/* Signed pointer */
+#define DW_EH_PE_SIGNED     0x08u
+/* Signed LE base-128 value */
+#define DW_EH_PE_SLEB128    0x09u
+/* Signed 16-bit value */
+#define DW_EH_PE_SDATA2     0x0Au
+/* Signed 32-bit value */
+#define DW_EH_PE_SDATA4     0x0Bu
+/* Signed 64-bit value */
+#define DW_EH_PE_SDATA8     0x0Cu
+typedef dw_uint DwEhPointerEncodingFormat;
+
+typedef enum 
+{
+  DW_EH_PE_ABSPTR  = 0x00, /* Absolute value */
+  DW_EH_PE_PCREL   = 0x10, /* Relative to address of encoded value */
+  DW_EH_PE_TEXTREL = 0x20, /* Text-Relative (GCC-specific) */
+  DW_EH_PE_DATAREL = 0x30, /* Data-Relative */
+  /* The following are not documented by LSB v1.3, yet they are used by
+     GCC, presumably they aren't documented by LSB since they aren't
+     used on Linux:  */
+  DW_EH_PE_FUNCREL = 0x40, /* Start of procedure relative */
+  DW_EH_PE_ALIGNED = 0x50  /* Aligned Pointer */
+} DwEHPointerEncodingApplication;
+
+typedef enum 
+{
+  DW_VIRTUAL_UNWIND_DATA_NULL,
+  DW_VIRTUAL_UNWIND_DATA_EH_FRAME,
+  DW_VIRTUAL_UNWIND_DATA_DEBUG_FRAME
+} DwVirtualUnwindDataType;
+
+#define DW_CIE_VERSION      3
+#define DW_CIE_VERSION_MAX  4
+
+typedef struct DwCommonInfoEntry 
+{
+  SymsAddr init_cfi_offset;
+  SymsAddr init_cfi_size;
+  SymsAddr end_offset;
+
+  SymsUWord cie_id;
+  SymsString augmentation;
+
+  syms_bool is_aug_sized;
+  syms_bool is_sig_frame;
+  syms_bool have_abi_maker;
+
+  U8 version;
+  U8 fde_encoding;
+  U8 lsda_encoding;
+  U8 handler_encoding;
+
+  S64 addr_size;
+  S64 segsel_size;
+
+  S64 code_align;
+  U64 data_align;
+  U64 ret_addr_reg;
+
+  SymsAddr handler_ip;
+} DwCommonInfoEntry;
+
+typedef enum DwCfiRegisterType 
+{
+  DW_CFI_REGISTER_TYPE_UNDEF,
+  DW_CFI_REGISTER_TYPE_SAME,
+  DW_CFI_REGISTER_TYPE_CFAREL,
+  DW_CFI_REGISTER_TYPE_REG,
+  DW_CFI_REGISTER_TYPE_EXPR,
+  DW_CFI_REGISTER_TYPE_VAL_EXPR,
+  DW_CFI_REGISTER_TYPE_INVALID
+} DwCfiRegisterType;
+
+typedef union DwCfiRegValue 
+{
+  SymsSWord w;
+  DwEncodedLocationExpr e;
+} DwCfiRegValue;
+
+static const dw_uint DW_CFA_COLUMN_REG = (DW_PRESERVED_REGS_MAX + 0);
+static const dw_uint DW_CFA_COLUMN_OFF = (DW_PRESERVED_REGS_MAX + 1);
+
+typedef struct DwCfiRow 
+{
+  DwCfiRegisterType type[DW_PRESERVED_REGS_MAX + 2];
+  DwCfiRegValue value[DW_PRESERVED_REGS_MAX + 2];
+} DwCfiRow;
+
+#define DW_CFI_PROGRAM_STACK_MAX 128
+typedef struct DwCfiProgram 
+{
+  DwCfiRow rules;
+
+  DwCfiRow *frame;
+  DwCfiRow stack[DW_CFI_PROGRAM_STACK_MAX];
+
+  dw_uint reg_count;
+  syms_bool setup_cfa;
+  S64 ret_addr_regid;
+  SymsAddr cfa;
+} DwCfiProgram;
+
+typedef struct DwFrameInfo 
+{
+  DwVirtualUnwindDataType source_type;
+  SymsAddr eh_frame;
+  SymsAddr image_base;
+  DwArchInfo arch_info;
+} DwFrameInfo;
+
+typedef struct DwVirtualUnwind 
+{
+  DwCfiProgram program;
+} DwVirtualUnwind;
+
+typedef struct DwFrameDescEntry 
+{
+  /* NOTE(nick): Offset where FDE was read */
+  SymsAddr data_off;
+
+  /* NOTE(nick): Offset of CFI for this FDE */
+  SymsAddr cfi_offset;
+
+  /* NOTE(nick): Number of bytes that FDE occupies */
+  U64 data_size;
+
+  /* NOTE(nick): Number of bytes that CFI take up */
+  U64 cfi_size;
+
+  /* NOTE(nick): Address of first instruction in procedure, that FDE describes. */
+  SymsAddr start_ip;
+
+  /* NOTE(nick): Number of bytes that instructions occupy. */
+  SymsAddr range_ip;
+
+  /* NOTE(nick): Address of LSDA. */
+  SymsAddr lsda_ip;
+} DwFrameDescEntry;
+
+typedef struct DwFrameDescEntryIter 
+{
+  DwFrameInfo *frame_info;
+  DwBinRead secdata;
+  DwCommonInfoEntry cie;
+  SymsAddr cie_offset;
+} DwFrameDescEntryIter;
+
+DW_API syms_bool
+dw_virtual_unwind_init(SymsImageType image_type, DwVirtualUnwind *context_out);
+
+DW_API syms_bool
+dw_virtual_unwind_frame(DwVirtualUnwind *context,
+            SymsArch arch,
+            DwVirtualUnwindDataType sec_bytes_type,
+            void *sec_bytes, SymsUMM sec_bytes_size,
+            SymsAddr image_base,
+            SymsAddr sec_bytes_base,
+            void *memread_ctx, dw_memread_sig *memread,
+            void *regread_ctx, dw_regread_sig *regread,
+            void *regwrite_ctx, dw_regwrite_sig *regwrite);
+
+SYMS_INTERNAL SymsAddr
+dw_parse_pointer(DwFrameInfo *frame_info, DwBinRead *secdata, dw_uint encoding);
+
+#endif /* SYMS_DWARF_UNWIND_INCLUDE_H */
diff --git a/Engine/Source/Runtime/SymsLib/src/elf/syms_elf.c b/Engine/Source/Runtime/SymsLib/src/elf/syms_elf.c
new file mode 100644
--- /dev/null
+++ b/Engine/Source/Runtime/SymsLib/src/elf/syms_elf.c
@@ -0,0 +1,352 @@
+// Copyright Epic Games, Inc. All Rights Reserved.
+/******************************************************************************
+ * File   : syms_elf.c                                                        *
+ * Author : Nikita Smith                                                      *
+ * Created: 2020/05/30                                                        *
+ * Purpose: Implementation of routines for dealing with ELF format            *
+ ******************************************************************************/
+
+#include "syms_elf.h"
+
+SYMS_COMPILER_ASSERT(sizeof(SymsSymtabIter) <= sizeof(SymsModImpl));
+
+SYMS_API syms_bool
+syms_img_init_elf(SymsImage *img, void *img_data, SymsUWord img_size, SymsLoadImageFlags flags)
+{
+  SymsImageElf *img_elf = (SymsImageElf *)img->impl;
+  SymsBuffer img_read = syms_buffer_init(img_data, img_size);
+  U8 sig[SYMS_EI_NIDENT];
+  u16 e_machine = SYMS_EM_NONE;
+  u32 i;
+  u32 phdr_num;
+  SymsAddr shstr_off;
+  SymsAddr base_addr = 0;
+  SymsAddr prev_delta;
+
+  if (!syms_buffer_read(&img_read, sig, sizeof(sig)))
+    return syms_false;
+  if (sig[SYMS_EI_MAG0] != 0x7f || 
+      sig[SYMS_EI_MAG1] != 'E'  ||
+      sig[SYMS_EI_MAG2] != 'L'  || 
+      sig[SYMS_EI_MAG3] != 'F')
+    return syms_false; // signature mismatch, data is not of ELF format
+
+
+  // seek back to origin, so we can read entire SymsElf32/SymsElf64 struct
+  syms_buffer_seek(&img_read, 0);
+  switch (sig[SYMS_EI_CLASS]) {
+  case SYMS_ELFCLASS32: {
+    SymsElf32 *elf = syms_buffer_push_struct(&img_read, SymsElf32);
+    SymsElfShdr32 *shstr;
+
+    if (!elf)
+      break;
+    // seek to section header that contains strings
+    shstr_off = elf->e_shoff + elf->e_shentsize*elf->e_shstrndx;
+    if (!syms_buffer_seek(&img_read, shstr_off))
+      break;
+    // read string section header
+    shstr = syms_buffer_push_struct(&img_read, SymsElfShdr32);
+    if (!shstr)
+      break;
+    if (shstr->sh_type == SYMS_SHT_NOBITS)
+      break;
+
+    // save base address of program-header names
+    img_elf->sh_name_lo = (flags & SYMS_LOAD_IMAGE_FLAGS_FROM_MEMORY) ? shstr->sh_addr : shstr->sh_offset;
+    img_elf->sh_name_hi = img_elf->sh_name_lo + shstr->sh_size;
+
+    // seek to first program header
+    if (!syms_buffer_seek(&img_read, elf->e_phoff))
+      break;
+    // search for base address, which is stored in program header that is closest to entry point
+    phdr_num = elf->e_phnum != SYMS_UINT16_MAX ? elf->e_phnum : shstr->sh_info;
+    prev_delta = SYMS_ADDR_MAX;
+    for (i = 0; i < phdr_num; ++i) {
+      SymsElfPhdr64 *phdr = syms_buffer_push_struct(&img_read, SymsElfPhdr64);
+      if (!phdr)
+        break;
+      if (phdr->p_type == SYMS_PT_LOAD) {
+        if (elf->e_entry >= phdr->p_vaddr) {
+          SymsAddr delta = elf->e_entry - phdr->p_vaddr;
+          if (delta < prev_delta) {
+            base_addr = phdr->p_vaddr;
+            prev_delta = delta;
+          }
+        }
+      }
+    }
+
+    if (elf->e_shoff >= img_size)
+      break;
+
+    e_machine = elf->e_machine;
+
+    img->type             = SYMS_IMAGE_ELF;
+    img->header_class     = SYMS_IMAGE_HEADER_CLASS_32;
+    img_elf->u.header32 = elf;
+  } break;
+  case SYMS_ELFCLASS64: {
+    SymsElf64 *elf = syms_buffer_push_struct(&img_read, SymsElf64);
+    SymsElfShdr64 *shstr;
+
+    if (!elf)
+      break;
+    // seek to .shstrtab
+    shstr_off = elf->e_shoff + elf->e_shentsize * elf->e_shstrndx;
+    if (!syms_buffer_seek(&img_read, shstr_off))
+      break;
+    // read string section header
+    shstr = syms_buffer_push_struct(&img_read, SymsElfShdr64);
+    if (shstr) {
+      if (shstr->sh_type == SYMS_SHT_NOBITS)
+        break;
+
+      // save base address of program-header names
+      img_elf->sh_name_lo = (flags & SYMS_LOAD_IMAGE_FLAGS_FROM_MEMORY) ? shstr->sh_addr : shstr->sh_offset;
+      img_elf->sh_name_hi = img_elf->sh_name_lo + shstr->sh_size;
+    } else {
+      img_elf->sh_name_lo = 0;
+      img_elf->sh_name_hi = 0;
+    }
+
+    // seek to first program header
+    if (!syms_buffer_seek(&img_read, elf->e_phoff))
+      break;
+    // search for base address, which is stored in program header that is closest to entry point
+    phdr_num = elf->e_phnum != SYMS_UINT16_MAX ? elf->e_phnum : shstr->sh_info;
+    prev_delta = SYMS_ADDR_MAX;
+    for (i = 0; i < phdr_num; ++i) {
+      SymsElfPhdr64 *phdr = syms_buffer_push_struct(&img_read, SymsElfPhdr64);
+      if (!phdr)
+        break;
+      if (phdr->p_type == SYMS_PT_LOAD) {
+        if (elf->e_entry >= phdr->p_vaddr) {
+          SymsAddr delta = elf->e_entry - phdr->p_vaddr;
+          if (delta < prev_delta) {
+            base_addr = phdr->p_vaddr;
+            prev_delta = delta;
+          }
+        }
+      }
+    }
+
+    if (elf->e_shoff >= img_size)
+      break;
+
+    e_machine = elf->e_machine;
+
+    img->type             = SYMS_IMAGE_ELF;
+    img->header_class     = SYMS_IMAGE_HEADER_CLASS_64;
+    img_elf->u.header64 = elf;
+  } break;
+  }
+
+  if (img->type == SYMS_IMAGE_ELF) {
+    img->base_addr = base_addr;
+
+    // convert e_machine to SymsArch
+    switch (e_machine) {
+    default: SYMS_ASSERT_FAILURE_PARANOID("undefined e_machine");
+    case SYMS_EM_AARCH64: img->arch = SYMS_ARCH_ARM;   break;
+    case SYMS_EM_ARM:     img->arch = SYMS_ARCH_ARM32; break;
+    case SYMS_EM_386:     img->arch = SYMS_ARCH_X86;   break;
+    case SYMS_EM_X86_64:  img->arch = SYMS_ARCH_X64;   break;
+    case SYMS_EM_PPC:     img->arch = SYMS_ARCH_PPC;   break;
+    case SYMS_EM_PPC64:   img->arch = SYMS_ARCH_PPC64; break;
+    case SYMS_EM_IA_64:   img->arch = SYMS_ARCH_IA64;  break;
+    }
+  }
+  
+  return (img->type == SYMS_IMAGE_ELF);
+}
+
+SYMS_INTERNAL SymsSecIterElf
+syms_sec_iter_init_elf(SymsImage *image)
+{
+  SymsImageElf *elf = (SymsImageElf *)image->impl;
+  SymsSecIterElf result;
+  SymsAddr secs_lo = 0;
+  SymsAddr secs_hi = 0;
+
+  result.image = image;
+  result.sh_name_lo = elf->sh_name_lo;
+  result.sh_name_hi = elf->sh_name_hi;
+  result.header_index = 0;
+  result.header_count = 0;
+
+  switch (image->header_class) {
+  case SYMS_IMAGE_HEADER_CLASS_NULL: {
+    syms_memzero(&result, sizeof(result));
+  } break;
+  case SYMS_IMAGE_HEADER_CLASS_32: {
+    secs_lo = elf->u.header32->e_shoff + elf->u.header32->e_shentsize;
+    secs_hi = secs_lo + elf->u.header32->e_shnum*sizeof(SymsElfShdr32);
+    result.header_count = elf->u.header32->e_shnum;
+  } break;
+  case SYMS_IMAGE_HEADER_CLASS_64: {
+    secs_lo = elf->u.header64->e_shoff;
+    secs_hi = secs_lo + elf->u.header64->e_shnum*sizeof(SymsElfShdr64);
+    result.header_count = elf->u.header64->e_shnum;
+  } break;
+  }
+
+  void *sec_data = (U8 *)image->img_data + secs_lo;
+  result.headers = syms_buffer_init(sec_data, image->img_data_size - secs_lo);
+
+  return result;
+}
+
+SYMS_INTERNAL syms_bool
+syms_sec_iter_next_elf(SymsSecIterElf *iter, SymsElfShdr64 *sec_out)
+{
+  SymsImage *image = iter->image;
+  syms_bool result = syms_false;
+
+  switch (image->header_class) {
+  case SYMS_IMAGE_HEADER_CLASS_NULL: {
+    syms_memzero(sec_out, sizeof(*sec_out));
+  } break;
+  case SYMS_IMAGE_HEADER_CLASS_32: {
+    // convert header to 64bit
+    SymsElfShdr32 *elf32 = syms_buffer_push_struct(&iter->headers, SymsElfShdr32);
+    if (elf32) {
+      sec_out->sh_name      = elf32->sh_name;
+      sec_out->sh_type      = elf32->sh_type;
+      sec_out->sh_flags     = elf32->sh_flags;
+      sec_out->sh_addr      = elf32->sh_addr;
+      sec_out->sh_offset    = elf32->sh_offset;
+      sec_out->sh_size      = elf32->sh_size;
+      sec_out->sh_link      = elf32->sh_link;
+      sec_out->sh_info      = elf32->sh_info;
+      sec_out->sh_addralign = elf32->sh_addralign;
+      sec_out->sh_entsize   = elf32->sh_entsize;
+      result = syms_true;
+    }
+  } break;
+  case SYMS_IMAGE_HEADER_CLASS_64: {
+    SymsElfShdr64 *elf64 = syms_buffer_push_struct(&iter->headers, SymsElfShdr64);
+    if (elf64) {
+      *sec_out = *elf64;
+      result = syms_true;
+    }
+  } break;
+  default: break;
+  } 
+#if 0
+  { 
+    sec_out->name = syms_string_init_lit("");
+    if (header.sh_name < iter->sh_name_lo && iter->sh_name_lo + header.sh_name < img->img_data_size) {
+      SymsBuffer buffer = syms_buffer_init(img->img_data, img->img_data_size);
+      if (syms_buffer_seek(&buffer, iter->sh_name_lo + header.sh_name)) {
+        sec_out->name = syms_buffer_read_string(&buffer);
+      }
+    }
+
+    sec_out->data_size = header.sh_size;
+    sec_out->off = header.sh_offset;
+    sec_out->va = header.sh_addr;
+    sec_out->u.elf = header;
+
+    // check if section has data in file
+    sec_out->data = 0;
+    if (header.sh_type != SYMS_SHT_NOBITS) {
+      SymsAddr data_lo = (img->flags & SYMS_LOAD_IMAGE_FLAGS_FROM_MEMORY) ? header.sh_addr : header.sh_offset;
+      SymsAddr data_hi = data_lo + header.sh_size;
+      if (data_hi <= img->img_data_size) {
+        sec_out->data = (void *)((U8 *)img->img_data + data_lo);
+      }
+    }
+
+    // advance to next section
+    iter->header_index += 1;
+
+    result = syms_true;
+  }
+#endif
+
+  return result;
+}
+
+ELF_API syms_bool
+syms_symtab_iter_init(struct SymsInstance *instance, SymsSymtabIter *iter)
+{
+  syms_bool result = syms_false;
+
+  SymsSection symtab;
+  SymsSection strtab;
+
+  syms_bool symtab_found = syms_img_sec_from_name(instance, syms_string_lit(".symtab"), &symtab);
+  syms_bool strtab_found = syms_img_sec_from_name(instance, syms_string_lit(".strtab"), &strtab);
+
+  syms_memzero(iter, sizeof(*iter));
+
+  if (symtab_found && strtab_found) {
+    SymsImage *img = &instance->img;
+    if (img->type == SYMS_IMAGE_ELF) {
+      switch (img->header_class) {
+      case SYMS_IMAGE_HEADER_CLASS_NULL: iter->count = 0; break;
+      case SYMS_IMAGE_HEADER_CLASS_32: iter->count = syms_trunc_u32(symtab.data_size / sizeof(SymsElf32Sym)); break;
+      case SYMS_IMAGE_HEADER_CLASS_64: iter->count = syms_trunc_u32(symtab.data_size / sizeof(SymsElf64Sym)); break;
+      }
+      iter->symtab_cursor = syms_buffer_init(symtab.data, symtab.data_size);
+      iter->strtab_cursor = syms_buffer_init(strtab.data, strtab.data_size);
+      iter->header_class = img->header_class;
+    }
+    result = syms_true;
+  }
+
+  return result;
+}
+
+ELF_API syms_bool
+syms_symtab_iter_next(SymsSymtabIter *iter, SymsSymtabEntry *entry_out)
+{
+  syms_bool result = syms_false;
+  switch (iter->header_class) {
+  case SYMS_IMAGE_HEADER_CLASS_NULL: break;
+  case SYMS_IMAGE_HEADER_CLASS_32: break;
+  case SYMS_IMAGE_HEADER_CLASS_64: {
+    SymsElf64Sym *sym = syms_buffer_push_struct(&iter->symtab_cursor, SymsElf64Sym);
+    if (sym) {
+      if (syms_buffer_seek(&iter->strtab_cursor, sym->st_name)) {
+        entry_out->name = syms_buffer_read_string(&iter->strtab_cursor);
+        entry_out->bind = SYMS_ELF64_ST_BIND(sym->st_info);
+        entry_out->type = SYMS_ELF64_ST_TYPE(sym->st_info);
+        entry_out->vis = SYMS_ELF64_ST_VISIBILITY(sym->st_info);
+        entry_out->index = sym->st_shndx;
+        entry_out->value = sym->st_value;
+        entry_out->size = sym->st_size;
+      } else {
+        entry_out->name = syms_string_init_lit("");
+        entry_out->bind = SYMS_STB_LOCAL;
+        entry_out->type = 0;
+        entry_out->vis = SYMS_STV_DEFAULT;
+        entry_out->index = sym->st_shndx;
+        entry_out->value = 0;
+        entry_out->size = 0;
+      }
+      result = syms_true;
+    }
+  } break;
+  }
+  return result;
+}
+
+SYMS_INTERNAL SymsProc
+syms_proc_from_stt_func(SymsSymtabEntry *stt_func)
+{
+  SymsProc result;
+  SYMS_ASSERT(stt_func->type == SYMS_STT_FUNC);
+
+  result.type_id = syms_type_id_null();
+  result.va = stt_func->value;
+  result.len = syms_trunc_u32(stt_func->size);
+  result.dbg_start_va = result.va;
+  result.dbg_end_va = result.va + result.len;
+  result.name_ref.type = SYMS_STRING_REF_TYPE_STR;
+  result.name_ref = syms_string_ref_str(stt_func->name); //.u.str = stt_func->name;
+
+  return result;
+}
+
