RAD symbol library

diff --git a/Engine/Source/Runtime/SymsLib/src/pdb/syms_dbi.h b/Engine/Source/Runtime/SymsLib/src/pdb/syms_dbi.h
new file mode 100644
--- /dev/null
+++ b/Engine/Source/Runtime/SymsLib/src/pdb/syms_dbi.h
@@ -0,0 +1,337 @@
+// Copyright Epic Games, Inc. All Rights Reserved.
+#ifndef SYMS_DBI_INCLUDE_H
+#define SYMS_DBI_INCLUDE_H
+
+typedef enum 
+{
+  PDB_DBI_VER_41  = 930803,
+  PDB_DBI_VER_50  = 19960307,
+  PDB_DBI_VER_60  = 19970606,
+  PDB_DBI_VER_70  = 19990903,
+  PDB_DBI_VER_110 = 20091201,
+
+  PDB_DBI_VER_CUR = PDB_DBI_VER_70
+} pdb_dbi_header_ver_e;
+
+#define PDB_DBI_SC_VER_60 (0xeffe0000ul + 19970605ul)
+#define PDB_DBI_SC_VER_2  (0xeffe0000ul + 20140516ul)
+typedef U32 pdb_dbi_sc_ver;
+
+#define PDB_DBI_HEADER_SIG_NULL 0
+#define PDB_DBI_HEADER_SIG_V1   0xFFFFFFFF
+typedef U32 pdb_dbi_header_sig;
+
+typedef U32 pdb_ni_t;
+
+/* name index */
+typedef U32 pdb_ni;
+
+typedef enum 
+{
+  PDB_DBG_STREAM_FPO,
+  PDB_DBG_STREAM_EXCEPTION,
+  PDB_DBG_STREAM_FIXUP,
+  PDB_DBG_STREAM_OMAP_TO_SRC,
+  PDB_DBG_STREAM_OMAP_FROM_SRC,
+  PDB_DBG_STREAM_SECTION_HEADER, // Stream contains array of Image_Section_Header.
+  PDB_DBG_STREAM_TOKEN_RDI_MAP,
+  PDB_DBG_STREAM_XDATA,
+  PDB_DBG_STREAM_PDATA,
+  PDB_DBG_STREAM_NEW_FPO,
+  PDB_DBG_STREAM_SECTION_HEADER_ORIG,
+
+  PDB_DBG_STREAM_MAX
+} pdb_dbg_stream_e;
+
+typedef struct pdb_sc40 
+{
+  pdb_isec sec;
+  U32 sec_off;
+  U32 size;
+  U32 flags;
+  pdb_imod imod;
+
+  U8 r[7];
+} pdb_sc40;
+
+typedef struct pdb_sc20 
+{
+  pdb_isec sec;
+  U32 sec_off;
+  U32 size;
+  pdb_imod imod;
+} pdb_sc20;
+
+typedef struct pdb_sc 
+{
+  pdb_isec sec;
+  U8 padding1[2];
+  pdb_isec_umm sec_off;
+  pdb_isec_umm size;
+  U32 flags;
+  pdb_imod imod;
+  U8 padding2[2];
+  U32 data_crc;
+  U32 reloc_crc;
+} pdb_sc;
+
+typedef struct pdb_sc2 
+{
+  pdb_isec sec;
+  U8 padding1[2];
+  pdb_isec_umm sec_off;
+  pdb_isec_umm size;
+  U32 flags;
+  pdb_imod imod;
+  U8 padding2[2];
+  U32 data_crc;
+  U32 reloc_crc;
+  U32 sec_coff;
+} pdb_sc2;
+
+#pragma pack(push, 1)
+
+enum
+{
+  PDB_DBI_HEADER_BUILD_MINOR_MASK = 0x00FF,
+  PDB_DBI_HEADER_BUILD_MINOR_SHIFT = 0,
+
+  PDB_DBI_HEADER_BUILD_MAJOR_MASK = 0x7F00,
+  PDB_DBI_HEADER_BUILD_MAJOR_SHIFT = 8,  
+
+  PDB_DBI_HEADER_NEW_VERSION_FORMAT_MASK = 0x8000
+};
+
+enum
+{
+  PDB_DBI_HEADER_FLAGS_INCREMENTAL_MASK = 0x1, // bit is set if linked incrementally
+  PDB_DBI_HEADER_FLAGS_STRIPPED_MASK    = 0x2, // bit is set if private symbols were stripped
+  PDB_DBI_HEADER_FLAGS_CTYPES_MASK      = 0x4  // bit is set if linked with /debug:ctypes
+};
+typedef u16 pdb_dbi_header_flags;
+
+typedef struct pdb_dbi_header 
+{
+  pdb_dbi_header_sig sig;
+  U32 version;
+  U32 age;
+  pdb_sn global_sym_sn;
+
+  U16 build_number;
+
+  /* public symbol stream. */
+  pdb_sn public_sym_sn;  
+
+  /* build version of the pdb dll that built this pdb last. */
+  U16 pdb_version;                
+
+  /* An array of symbols, nothing special (Format u16 size, u16 type, char data[size]) */
+  pdb_sn sym_record_sn;  
+  
+  /* rbld version of the pdb dll that built this pdb last. */
+  U16 pdb_version2;               
+
+  U32 module_info_size;
+  U32 sec_con_size;
+  U32 sec_map_size;
+  U32 file_info_size;
+
+  /* size of the Type Server Map substream. */
+  U32 tsm_size;           
+
+  /* index of MFC type server */
+  U32 mfc_index;          
+
+  /* size of optional DbgHdr info appended to the end of the stream */
+  U32 dbg_header_size;    
+  
+  /* number of bytes in EC substream, or 0 if EC no EC enabled Mods */
+  U32 ec_info_size;       
+
+  U16 flags;
+  
+  /* ImageFileMachine */
+  U16 machine;     
+
+  U32 reserved;
+} pdb_dbi_header;
+
+typedef struct pdb_mod_header 
+{
+  U8 unused[4]; 
+  
+  pdb_sc sc;
+
+  U16 flags;
+
+  /* stream number of module debug info (syms, lines, FPO) */
+  pdb_sn sn;    
+
+  /* size of local symbol debug info in above stream */
+  U32 symbol_bytes;      
+
+  /* size of line number debug info in stream sn */
+  U32 c11_lines_size;    
+
+  /* size of C13 style line number info in stream sn */
+  U32 c13_lines_size;    
+
+  /* number of files contributing to this module (file index) */
+  U16 num_contrib_files; 
+
+  U8 padding1[2];
+
+  /* Shit field that is unused in our case */
+  U32 file_names_offset; 
+
+  pdb_ni src_file_ni;
+  pdb_ni pdb_file_ni;
+
+  /* char module_name[]; */
+  /* char obj_name[]; */
+} pdb_mod_header;
+
+typedef struct pdb_psi_header 
+{
+  U32 sym_hash_size;
+  U32 addr_map_size;
+  U32 thunk_count;
+  U32 thunk_size; 
+  pdb_isec isec_thunk_table;
+  U8 padding[2];
+  pdb_isec_umm sec_thunk_tabl_off; 
+  U32 sec_count;
+} pdb_psi_header;
+
+#define PDB_GSI_V70   0xeffe0000 + 19990810
+#define PDB_GSI_SIG   (~0u)
+typedef struct pdb_gsi_header 
+{
+  U32 sig;
+  U32 ver;
+  U32 hr_len;
+  U32 num_buckets;
+} pdb_gsi_header;
+
+typedef struct pdb_gsi_file_hr 
+{
+  U32 off; /* Offset in the symbol record stream */
+  U32 cref;
+} pdb_gsi_file_hr;
+
+typedef struct pdb_gsi_hr 
+{
+  U32 off;
+  struct pdb_gsi_hr *next;
+} pdb_gsi_hr;
+
+typedef struct pdb_ti_off 
+{
+  pdb_ti ti;
+  U32 off;
+} pdb_ti_off;
+
+#define PDB_INTV_VC2 0xE0D5C
+
+#define PDB_TM_IMPV40          0x1306B4A
+#define PDB_TM_IMPV41          0x1306E12
+#define PDB_TM_IMPV50_INTERIM  0x13091F3
+#define PDB_TM_IMPV50          0x13094C7
+#define PDB_TM_IMPV70          0x1310977
+#define PDB_TM_IMPV80          0x131CA0B
+#define PDB_TM_IMV_CURR        PDB_TM_IMV80
+typedef struct pdb_tm_header 
+{
+  U32 version;
+
+  /* Size of this header */
+  U32 header_size;
+
+  /* All types below the "ti_lo" are reserved for internal use. Most of them
+   * represent basic types, like int, float, char and so on. */
+  U32 ti_lo;
+  U32 ti_hi;
+
+  /* Size in bytes of all types that follow pdb_tm_header. */
+  U32 types_size;
+
+  /* Hash stream is sub-divided into following sections:
+   *
+   * pdb_ti hashes[ti_hi - ti_lo] 
+   *
+   * pdb_ti_off offsets[ti_off.cb/sizeof(pdb_ti_off)] 
+   *
+   * Last section is a serialized map of types that have to be moved in the
+   * internally-linked table (pdb_tm.buckets). Size of this sub-section
+   * is defined by the "hash_adj" field below. */
+  pdb_sn hash_sn;
+
+  /* If stream not present this value is 0xffff */
+  pdb_sn hash_sn_aux; 
+
+  /* Size of the hash value in the "hash_sn". */
+  S32 hash_key_size;
+
+  /* Count of hashes that map string name of 
+   * a type to it's corresponding type index. */
+  U32 hash_bucket_count;
+
+  pdb_off_cb hash_vals;
+  pdb_off_cb ti_off;
+  pdb_off_cb hash_adj;
+} pdb_tm_header;
+
+typedef struct pdb_img_sec 
+{
+  char name[8];
+
+  union {
+    U32 physical_address;
+    U32 virtual_size;
+  } misc;
+
+  U32 rva;
+  U32 sizeof_raw_data;
+  U32 pointer_to_raw_data;
+  U32 pointer_to_relocations;
+  U32 pointer_to_linenumbers;
+  U16 relocations_count;
+  U16 linenumbers_count;
+  U32 characteristics;
+} pdb_img_sec;
+
+enum
+{
+  PDB_OMF_SEG_DESC_FLAGS_NULL          = 0,
+  PDB_OMF_SEG_DESC_FLAGS_READ          = (1u << 0u),
+  PDB_OMF_SEG_DESC_FLAGS_WRITE         = (1u << 1u),
+  PDB_OMF_SEG_DESC_FLAGS_EXEC          = (1u << 2u),
+  PDB_OMF_SEG_DESC_FLAGS_ADDR_IS_32BIT = (1u << 3u),
+  PDB_OMF_SEG_DESC_FLAGS_IS_SELECTOR   = (1u << 8u),
+  PDB_OMF_SEG_DESC_FLAGS_IS_ABS_ADDR   = (1u << 9u),
+  PDB_OMF_SEG_DESC_FLAGS_IS_GROUP      = (1u << 10u)
+};
+typedef u16 pdb_omf_flags;
+
+typedef struct pdb_secmap_header
+{
+  u16 sec_count;      // Count of segment descriptors in the table
+  u16 sec_count_log;  // Count of logical segments in the table
+} pdb_secmap_header;
+
+typedef struct pdb_secmap_entry
+{
+  u16 flags;           // Entry flags. See pdb_omf_flags
+  u16 ovl;             // Logical overlay number
+  u16 group;           // Index of group in the descriptor array
+  u16 frame;           // ??
+  u16 sec_name;        // Index in the sstSegName table, or if not present SYMS_UINT16_MAX
+  u16 class_name;      // Index in the sstSegName table, or if not present SYMS_UINT16_MAX
+  u32 offset;          // Logical segment offset within the physical segment
+  u32 sec_byte_length; // Number of bytes that make up group or segment
+} pdb_secmap_entry;
+
+#pragma pack(pop)
+
+#endif /* SYMS_DBI_INCLUDE_H */
+
diff --git a/Engine/Source/Runtime/SymsLib/src/pdb/syms_msf.h b/Engine/Source/Runtime/SymsLib/src/pdb/syms_msf.h
new file mode 100644
--- /dev/null
+++ b/Engine/Source/Runtime/SymsLib/src/pdb/syms_msf.h
@@ -0,0 +1,161 @@
+// Copyright Epic Games, Inc. All Rights Reserved.
+#ifndef SYMS_MSF_INCLUDE_H
+#define SYMS_MSF_INCLUDE_H
+
+#ifndef PDB_API
+#define PDB_API
+#endif
+
+typedef U16 pdb_uint16;
+typedef S32 pdb_int;
+typedef U32 pdb_uint;
+typedef U32 pdb_umm;
+#define PDB_UINT_MAX SYMS_UINT32_MAX
+#define pdb_trunc_uint(x) (pdb_uint)syms_trunc_u32(x)
+#define pdb_trunc_int(x)  (pdb_int)syms_trunc_u32(x)
+
+struct pdb_context;
+
+typedef u32 pdb_offset;
+
+/* stream number */
+typedef u16 pdb_sn;
+
+#define PDB_INVALID_SN 0xffff
+
+typedef struct pdb_header20 
+{
+  U8 magic[44];
+  U32 page_size;
+  U16 free_page_map;
+  U16 pages_used;
+  U32 root_size;
+  U32 reserved;
+  U16 page_map_addr;
+} pdb_header20;
+
+typedef struct pdb_header70 
+{
+  U8 magic[32];
+
+  /* Size of page in bytes */
+  U32 page_size;
+
+  /* Index of the free page map */
+  S32 free_page_map;
+
+  /* Number of pages used */
+  S32 pages_used;
+
+  /* Size of root directory in bytes */
+  U32 root_size;
+
+  /* Purpose of this field is unknown */
+  U32 reserved;
+
+  /* Number of page that contains page map */
+  U32 page_map_addr;
+} pdb_header70;
+
+typedef enum 
+{
+  PDB_PAGE_INDEX16 = 2,
+  PDB_PAGE_INDEX32 = 4
+} pdb_page_index_type;
+
+#define PDB_STREAM_FLAGS_READ_FAILED (1u << 0u)
+#define PDB_STREAM_FLAGS_SEEK_FAILED (1u << 1u)
+typedef U32 pdb_stream_flags;
+
+#define PDB_STREAM_READ_OR_SEEK_FAILED(flags) (((flags) & PDB_STREAM_FLAGS_READ_FAILED) || ((flags) & PDB_STREAM_FLAGS_SEEK_FAILED))
+#define PDB_STREAM_NO_ERROR(flags)            ((~(flags) & PDB_STREAM_FLAGS_READ_FAILED) && (~(flags) & PDB_STREAM_FLAGS_SEEK_FAILED))
+typedef struct pdb_stream 
+{
+  struct pdb_context *pdb;
+
+  pdb_stream_flags flags;
+
+  /* Stream number (ID) */
+  pdb_sn sn;
+
+  U32 page_size;
+  U32 page_read_lo;
+  U32 page_read_hi;
+
+  /* Indicates how many bytes read into the stream. */
+  U32 off;
+
+  /* Number of bytes read when stream was subset */
+  U32 off_at_subset;
+
+  /* Amount of bytes that can be read from this stream. */
+  U32 size;
+
+  /* Offset into the root stream where the page indices for this stream start. */
+  U32 indices_off;
+
+  /* Current location in the root stream. */
+  U32 root_off;
+} pdb_stream;
+
+typedef struct pdb_symrec 
+{
+  U16 size;
+  U16 type;
+  U32 end;
+} pdb_symrec;
+
+#define pdb_stream_skip(stream, num) pdb_stream_seek(stream, (stream)->off + (num))
+PDB_API syms_bool
+pdb_stream_seek(pdb_stream *stream, U32 offset);
+
+#define pdb_stream_read_struct(stream, bf, type) pdb_stream_read(stream, bf, sizeof(type))
+PDB_API pdb_uint
+pdb_stream_read(pdb_stream *stream, void *buff, pdb_uint buff_max);
+
+PDB_API syms_bool
+pdb_stream_can_read_bytes(pdb_stream *stream, pdb_uint num_bytes);
+
+PDB_API syms_bool
+pdb_stream_align(pdb_stream *stream, pdb_uint align);
+
+PDB_API syms_bool
+pdb_stream_init(struct pdb_context *pdb, pdb_sn sn, pdb_stream *out_stream);
+
+PDB_API syms_bool
+pdb_stream_init_at(struct pdb_context *pdb, pdb_sn sn, pdb_offset offset, pdb_stream *out_stream);
+
+PDB_API void
+pdb_stream_init_null(pdb_stream *stream);
+
+PDB_API syms_bool
+pdb_stream_is_null(pdb_stream *stream);
+
+PDB_API pdb_stream
+pdb_stream_subset(const pdb_stream *stream, pdb_uint off, pdb_uint size);
+
+PDB_API pdb_uint
+pdb_stream_get_abs_off(pdb_stream *stream);
+
+PDB_API syms_bool
+pdb_stream_read_u08(pdb_stream *stream, U8 *out_value);
+
+PDB_API syms_bool
+pdb_stream_read_u16(pdb_stream *stream, U16 *out_value);
+
+PDB_API syms_bool
+pdb_stream_read_u32(pdb_stream *stream, pdb_uint *out_value);
+
+PDB_API syms_bool
+pdb_stream_read_uleb32(pdb_stream *stream, pdb_uint *out_value);
+
+PDB_API syms_bool
+pdb_stream_read_symrec(pdb_stream *stream, pdb_symrec *out_rec);
+
+PDB_API pdb_uint
+pdb_stream_read_str(pdb_stream *stream, void *bf, pdb_uint bf_max);
+
+PDB_API pdb_uint
+pdb_stream_strlen(pdb_stream *stream);
+
+#endif /* SYMS_MSF_INCLUDE_H */
