RAD symbol library

diff --git a/Engine/Source/Runtime/SymsLib/src/pdb/syms_pdb.c b/Engine/Source/Runtime/SymsLib/src/pdb/syms_pdb.c
new file mode 100644
--- /dev/null
+++ b/Engine/Source/Runtime/SymsLib/src/pdb/syms_pdb.c
@@ -0,0 +1,6438 @@
+// Copyright Epic Games, Inc. All Rights Reserved.
+/* -------------------------------------------------------------------------------- 
+ *                                   Multi-Stream File (MSF)
+ * -------------------------------------------------------------------------------- */
+
+PDB_API U32
+pdb_count_pages(pdb_context *pdb, U32 length)
+{
+  U32 page_count = length / pdb->page_size;
+  if (length % pdb->page_size) {
+    page_count += 1;
+  }
+  return page_count;
+}
+
+PDB_API syms_bool
+pdb_can_read_bytes(pdb_context *pdb, void *ptr_start, U32 num_bytes)
+{
+  void *file_start = (void *)pdb->file_data;
+  void *file_end   = (U8 *)pdb->file_data + pdb->file_size;
+  void *ptr_end    = (syms_bool *)ptr_start + num_bytes;
+  syms_bool result  = ptr_start >= file_start && ptr_end <= file_end;
+  return result;
+}
+
+PDB_API syms_bool
+pdb_stream_root_seek(pdb_stream *stream, U32 off)
+{
+  syms_bool result = syms_false;
+  if (off <= stream->pdb->root_size) {
+    stream->root_off = off;
+    result = syms_true;
+  }
+  return result;
+}
+
+PDB_API pdb_uint
+pdb_root_read(pdb_context *pdb, pdb_uint off, pdb_uint data_size, void *data_out)
+{
+  // TODO(nick): detect contiguous pages and extract data with a single read
+  pdb_uint read_size = 0;
+  if (off + data_size <= pdb->root_size) {
+    U8 *file_data = (U8 *)pdb->file_data;
+    pdb_uint chunk_count = pdb_count_pages(pdb, data_size);
+    pdb_uint chunk_index;
+    for (chunk_index = 0; chunk_index < chunk_count; ++chunk_index) {
+      pdb_uint page_index = off / pdb->page_size;
+      pdb_uint page_index_max = pdb->page_size / pdb->page_index_size;
+      pdb_uint addr_index = (page_index * pdb->page_index_size) / pdb->page_size;
+      pdb_uint data_page_index = PDB_UINT_MAX;
+      switch (pdb->page_index_size) {
+      case 4: {
+        pdb_uint root_index = pdb->page_map_addr.addr32[addr_index];
+        pdb_uint root_index_off = (root_index * pdb->page_size) + ((page_index % page_index_max) * pdb->page_index_size);
+        if (root_index_off + pdb->page_index_size > pdb->file_size) {
+          root_index = PDB_UINT_MAX;
+          break;
+        }
+        syms_memcpy(&data_page_index, file_data + root_index_off, pdb->page_index_size);
+      } break;
+      case 2: {
+        data_page_index = (U32)pdb->page_map_addr.addr16[addr_index]; 
+      } break;
+      }
+      {
+        pdb_uint data_off = data_page_index * pdb->page_size + (off % pdb->page_size);
+        pdb_uint to_read = SYMS_MIN(data_size, pdb->page_size - (off % pdb->page_size));
+        if (data_off + to_read > pdb->file_size)
+          break;
+        syms_memcpy((U8 *)data_out + read_size, file_data + data_off, to_read);
+        read_size += to_read;
+        data_size -= to_read;
+      }
+    }
+  }
+  return read_size;
+}
+
+PDB_API syms_bool
+pdb_stream_root_read_u16(pdb_stream *stream, U16 *value_out)
+{
+  U32 read_size = pdb_root_read(stream->pdb, stream->root_off, sizeof(*value_out), (void *)value_out);
+  if (read_size == sizeof(*value_out)) {
+    stream->root_off += sizeof(*value_out);
+    return syms_true;
+  }
+  return syms_false;
+}
+PDB_API syms_bool
+pdb_stream_root_read_u32(pdb_stream *stream, U32 *value_out)
+{
+  U32 read_size = pdb_root_read(stream->pdb, stream->root_off, sizeof(*value_out), (void *)value_out);
+  if (read_size == sizeof(*value_out)) {
+    stream->root_off += sizeof(*value_out);
+    return syms_true;
+  }
+  return syms_false;
+}
+
+PDB_API void
+pdb_stream_init_null(pdb_stream *stream)
+{
+  stream->pdb             = 0;
+  stream->flags           = 0;
+  stream->sn              = PDB_INVALID_SN;
+  stream->page_size       = 0x1000;
+  stream->page_read_lo    = 0;
+  stream->page_read_hi    = 0;
+  stream->off             = PDB_UINT_MAX;
+  stream->off_at_subset   = 0;
+  stream->size            = 0;
+  stream->indices_off     = PDB_UINT_MAX;
+  stream->root_off        = PDB_UINT_MAX;
+}
+PDB_API syms_bool
+pdb_stream_is_null(pdb_stream *stream)
+{
+  return (stream->pdb == 0);
+}
+
+
+PDB_API syms_bool
+pdb_stream_init(pdb_context *pdb, pdb_sn sn, pdb_stream *stream_out)
+{
+  return pdb_stream_init_at(pdb, sn, 0, stream_out);
+}
+
+PDB_API syms_bool
+pdb_stream_init_at(pdb_context *pdb, pdb_sn sn, pdb_offset at_offset, pdb_stream *stream)
+{
+  U32 i;
+  u32 num_streams;
+
+  pdb_stream_init_null(stream);
+
+  stream->pdb       = pdb;
+  stream->page_size = pdb->page_size;
+  stream->sn        = sn;
+  stream->root_off  = 0;
+  stream->flags = 0;
+
+  if (!pdb_stream_root_read_u32(stream, &num_streams)) {
+#if defined(SYMS_PARANOID)
+    SYMS_ASSERT_CORRUPTED_STREAM;
+#endif
+    return syms_false;
+  }
+  if (sn >= num_streams) {
+    return syms_false;
+  }
+
+  /* Read sizes of the previous streams to figure out where page indices
+   * for this stream start. */
+  switch (pdb->page_index_size) {
+  case 4: {
+    /* Root stream layouts:
+     * U32 num_streams;
+     * U32 stream_sizes[num_streams];
+     * U32 stream_page_indices[num_streams][]; 
+     */
+    /* Skip stream count. */
+    stream->indices_off = sizeof(U32);
+    /* Skip stream sizes. */
+    stream->indices_off += num_streams*sizeof(U32);
+    for (i = 0; i < sn; ++i) {
+      U32 stream_size;
+      if (!pdb_stream_root_read_u32(stream, &stream_size)) {
+#if defined(SYMS_PARANOID)
+        SYMS_ASSERT_CORRUPTED_STREAM;
+#endif
+        return syms_false;
+      }
+      if (stream_size == 0xffffffff) {
+        stream_size = 0;
+      }
+      stream->indices_off += pdb_count_pages(pdb, stream_size)*pdb->page_index_size;
+    }
+  } break;
+  case 2: {
+    /* Skip stream count. */
+    stream->indices_off = sizeof(U32);
+    /* Skip stream sizes. */
+    stream->indices_off += num_streams*sizeof(U32)*2;
+    for (i = 0; i < sn; ++i) {
+      U32 stream_size;
+      U32 unknown;
+
+      if (!pdb_stream_root_read_u32(stream, &stream_size)) {
+        SYMS_ASSERT_CORRUPTED_STREAM;
+        return syms_false;
+      }
+      if (!pdb_stream_root_read_u32(stream, &unknown)) {
+        SYMS_ASSERT_CORRUPTED_STREAM;
+        return syms_false;
+      }
+      if (stream_size == 0xffffffff) {
+        stream_size = 0;
+      }
+      stream->indices_off += pdb_count_pages(pdb, stream_size)*pdb->page_index_size;
+    }
+  } break;
+  default: SYMS_INVALID_CODE_PATH;
+  }
+
+  if (!pdb_stream_root_read_u32(stream, &stream->size)) {
+    SYMS_ASSERT_CORRUPTED_STREAM;
+    return syms_false;
+  }
+
+  /* Point root stream where the page indices for this sub-stream start. */
+  if (!pdb_stream_root_seek(stream, stream->indices_off)) {
+    SYMS_ASSERT_CORRUPTED_STREAM;
+    return syms_false;
+  }
+
+  /* Point this sub-stream to a first page. */
+  stream->off = stream->size;
+  return pdb_stream_seek(stream, at_offset);
+}
+
+PDB_API syms_bool
+pdb_stream_seek_(pdb_stream *stream, U32 off)
+{
+  syms_bool is_seek_valid = syms_false;
+  pdb_context *pdb = stream->pdb;
+  U32 page_size = stream->page_size;
+  U32 index = off / page_size;
+  U32 index_off = stream->indices_off + index*pdb->page_index_size;
+
+  stream->flags |= PDB_STREAM_FLAGS_SEEK_FAILED;
+  if (pdb_root_read(pdb, index_off, pdb->page_index_size, (void *)&index) == pdb->page_index_size) {
+    U32 page_read_lo = (index * page_size);
+    U32 page_read_hi = page_read_lo + page_size;
+    page_read_lo += off % page_size;
+
+    if (page_read_hi <= pdb->file_size) {
+      stream->page_read_lo = page_read_lo;
+      stream->page_read_hi = page_read_hi;
+      stream->flags &= ~(PDB_STREAM_FLAGS_SEEK_FAILED|PDB_STREAM_FLAGS_READ_FAILED);
+      is_seek_valid = syms_true;
+    } else {
+      SYMS_ASSERT_FAILURE_PARANOID("invalid page index");
+    }
+  } else {
+    //SYMS_ASSERT_FAILURE_PARANOID("invalid root-read");
+  }
+  return is_seek_valid;
+}
+PDB_API syms_bool
+pdb_stream_seek(pdb_stream *stream, U32 off)
+{
+  SYMS_ASSERT_PARANOID(off <= stream->size);
+  off = SYMS_MIN(off, stream->size);
+
+  if (stream->size == 0) {
+    stream->off = off;
+    return syms_true;
+  }
+
+  if (stream->size > 0) {
+    U32 off_abs = stream->off_at_subset + off;
+    U32 page_lo = stream->off_at_subset + stream->off;
+    U32 page_hi = page_lo + (stream->page_read_hi - stream->page_read_lo);
+    if (off_abs >= page_lo && off_abs < page_hi) {
+      stream->off = off;
+      stream->page_read_lo = (stream->page_read_hi - stream->page_size) + off_abs % stream->page_size;
+      return syms_true;
+    } else {
+      if (pdb_stream_seek_(stream, off_abs)) {
+        stream->off = off;
+        return syms_true;
+      }
+    }
+  }
+
+  return syms_false;
+}
+
+PDB_API pdb_context *
+pdb_stream_get_pdb_context(pdb_stream *stream)
+{
+  pdb_context *pdb = stream->pdb;
+  if (pdb) {
+    if (pdb->file_data == 0 && pdb->file_size == 0) {
+      pdb = 0;
+    }
+  }
+  return pdb;
+}
+
+PDB_API U32
+pdb_stream_read_(pdb_stream *stream, U32 off, void *dst, U32 dst_size)
+{
+  pdb_context *pdb = stream->pdb;
+  U32 read_size = 0;
+
+  if (pdb_stream_seek_(stream, off)) {
+    while (read_size < dst_size) {
+      U32 data_size = (stream->page_read_hi - stream->page_read_lo);
+      U32 data_copy_size = SYMS_MIN(data_size, (dst_size - read_size));
+
+      SYMS_ASSERT(read_size + data_copy_size <= dst_size);
+      syms_memcpy((U8 *)dst + read_size,
+          (U8 *)pdb->file_data + stream->page_read_lo,
+          data_copy_size);
+
+      SYMS_ASSERT(data_copy_size > 0);
+      read_size += data_copy_size;
+
+      if (!pdb_stream_skip(stream, data_copy_size)) {
+        read_size = 0;
+        break;
+      }
+    }
+  }
+
+  return read_size;
+}
+
+PDB_API U32
+pdb_stream_read(pdb_stream *stream, void *buffer, U32 buffer_size)
+{
+  U32 read_size = SYMS_MIN(stream->size - stream->off, buffer_size);
+  read_size = pdb_stream_read_(stream, (stream->off + stream->off_at_subset), buffer, read_size);
+  if (buffer_size > 0) {
+    if (read_size == 0) {
+      stream->flags |= PDB_STREAM_FLAGS_READ_FAILED;
+    } else {
+      stream->flags &= ~PDB_STREAM_FLAGS_READ_FAILED;
+    }
+  }
+  return read_size;
+}
+
+PDB_API U32
+pdb_stream_read_utf8(pdb_stream *stream)
+{
+  /* TODO(nick): Implement UTF8 */
+  U8 byte = 0;
+  U32 codepoint = 0;
+  pdb_stream_read_u08(stream, &byte);
+  codepoint = byte;
+  return codepoint;
+}
+
+PDB_API syms_bool
+pdb_stream_read_u08(pdb_stream *stream, U8 *out_value)
+{
+  U32 read = pdb_stream_read(stream, out_value, sizeof(U8));
+  syms_bool result = read == sizeof(U8);
+  return result;
+}
+
+PDB_API syms_bool
+pdb_stream_read_u16(pdb_stream *stream, U16 *out_value)
+{
+  U32 read = pdb_stream_read(stream, out_value, sizeof(U16));
+  syms_bool result = read == sizeof(U16);
+  return result;
+}
+
+PDB_API syms_bool
+pdb_stream_read_u32(pdb_stream *stream, U32 *out_value)
+{
+  U32 read = pdb_stream_read(stream, out_value, sizeof(U32));
+  syms_bool result = read == sizeof(U32);
+  return result;
+}
+
+PDB_API syms_bool
+pdb_stream_read_s32(pdb_stream *stream, S32 *out_value)
+{
+  U32 read = pdb_stream_read(stream, out_value, sizeof(S32));
+  syms_bool result = read == sizeof(S32);
+  return result;
+}
+
+PDB_API syms_bool
+pdb_stream_read_uint(pdb_stream *stream, pdb_uint *value_out)
+{
+  *value_out = 0;
+#if defined(PDB_64)
+  return pdb_stream_read_u64(stream, value_out);
+#else
+  return pdb_stream_read_u32(stream, value_out);
+#endif
+}
+
+PDB_API syms_bool
+pdb_stream_read_int(pdb_stream *stream, pdb_int *value_out)
+{
+  *value_out = 0;
+#if defined(PDB_64)
+  return pdb_stream_read_s64(stream, value_out);
+#else
+  return pdb_stream_read_s32(stream, value_out);
+#endif
+}
+
+PDB_API syms_bool
+pdb_stream_read_uleb32(pdb_stream *stream, U32 *value_out)
+{
+  U32 result = 0;
+  U8 b;
+  syms_bool was_parsed = syms_false;
+
+  if (pdb_stream_read_u08(stream, &b)) {
+    if ((b & 0x80) == 0x00) {
+      result = (U32)b;
+      was_parsed = syms_true;
+    } else if ((b & 0xC0) == 0x80) {
+      result = (U32)(b & 0x3f) << 8;
+      if (pdb_stream_read_u08(stream, &b)) {
+        result |= b;
+        was_parsed = syms_true;
+      }
+    } else if ((b & 0xE0) == 0xC0) {
+      result = (U32)(b & 0x1f) << 24;
+      if (pdb_stream_read_u08(stream, &b)) {
+        result |= (U32)b << 16;
+        if (pdb_stream_read_u08(stream, &b)) {
+          result |= (U32)b;
+          was_parsed = syms_true;
+        }
+      }
+    }
+  }
+
+  if (value_out) {
+    *value_out = result;
+  }
+
+  return was_parsed;
+}
+
+PDB_API syms_bool
+pdb_stream_read_sleb32(pdb_stream *stream, S32 *value_out)
+{
+  U32 uleb32;
+  S32 sleb32;
+
+  syms_bool was_parsed = syms_false;
+  if (pdb_stream_read_uleb32(stream, &uleb32)) {
+    if (uleb32 & 1) {
+      sleb32 = -(S32)(uleb32 >> 1);
+    } else {
+      sleb32 = (S32)(uleb32 >> 1);
+    }
+    was_parsed = syms_true;
+
+    if (value_out) {
+      *value_out = sleb32;
+    }
+  }
+
+  return was_parsed;
+}
+
+PDB_API syms_bool
+pdb_stream_read_symrec(pdb_stream *stream, pdb_symrec *out_rec)
+{
+  syms_bool result = syms_false;
+
+  if (pdb_stream_can_read_bytes(stream, sizeof(U16) + sizeof(U16))) {
+    pdb_stream_read_u16(stream, &out_rec->size);
+    if (out_rec->size >= sizeof(u16) && pdb_stream_can_read_bytes(stream, out_rec->size)) {
+      u32 symbol_start = stream->off;
+      pdb_stream_read_u16(stream, &out_rec->type);
+      out_rec->end = symbol_start + out_rec->size;
+      out_rec->size -= sizeof(u16);
+      result = syms_true;
+    }
+  }
+
+  return result;
+}
+
+PDB_API syms_bool
+pdb_stream_read_numeric(pdb_stream *stream, pdb_numeric *out_num)
+{
+  syms_bool is_read = syms_false;
+  U16 type;
+
+  if (pdb_stream_read_u16(stream, &type)) {
+    pdb_numeric num;
+
+    syms_memzero(num.u.data, sizeof(num.u.data));
+
+    if (type < PDB_LF_NUMERIC) {
+      num.itype = PDB_BASIC_TYPE_USHORT;
+      num.u.uint16 = type;
+      is_read = syms_true;
+    } else {
+      u32 itype;
+      u32 itype_size;
+      u32 read_size;
+      switch (type) {
+      case PDB_LF_REAL16:    itype = (PDB_BASIC_TYPE_REAL16  | (2 << 8));  break;
+      case PDB_LF_REAL32:    itype = (PDB_BASIC_TYPE_REAL32  | (4 << 8));  break;
+      case PDB_LF_REAL48:    itype = (PDB_BASIC_TYPE_REAL48  | (6 << 8));  break;
+      case PDB_LF_REAL64:    itype = (PDB_BASIC_TYPE_REAL64  | (8 << 8));  break;
+      case PDB_LF_REAL80:    itype = (PDB_BASIC_TYPE_REAL80  | (10 << 8)); break;
+      case PDB_LF_REAL128:   itype = (PDB_BASIC_TYPE_REAL128 | (12 << 8)); break;
+      case PDB_LF_CHAR:      itype = (PDB_BASIC_TYPE_CHAR    | (1 << 8));  break;
+      case PDB_LF_SHORT:     itype = (PDB_BASIC_TYPE_SHORT   | (2 << 8));  break;
+      case PDB_LF_USHORT:    itype = (PDB_BASIC_TYPE_USHORT  | (2 << 8));  break;
+      case PDB_LF_LONG:      itype = (PDB_BASIC_TYPE_LONG    | (4 << 8));  break;
+      case PDB_LF_ULONG:     itype = (PDB_BASIC_TYPE_ULONG   | (4 << 8));  break;
+      case PDB_LF_UQUADWORD: itype = (PDB_BASIC_TYPE_UQUAD   | (8 << 8));  break;
+      case PDB_LF_QUADWORD:  itype = (PDB_BASIC_TYPE_QUAD    | (8 << 8));  break;
+      default:               itype = PDB_BASIC_TYPE_NOTYPE;                break;
+      }
+      itype_size = PDB_BASIC_TYPE_SIZE_MASK(itype);
+      SYMS_ASSERT(itype_size <= sizeof(num.u.data));
+      read_size = pdb_stream_read(stream, num.u.data, itype_size);
+      is_read = (itype_size == read_size);
+    }
+    if (is_read && out_num) {
+      *out_num = num;
+    }
+  }
+
+  return is_read;
+}
+
+PDB_API syms_bool
+pdb_stream_read_numeric_u32(pdb_stream *stream, U32 *out_value)
+{
+  pdb_numeric num;
+  syms_bool is_read = pdb_stream_read_numeric(stream, &num);
+  if (is_read && out_value) {
+    *out_value = num.u.uint32;
+  }
+  return is_read;
+}
+
+// read string using raw page offsets cause it is faster
+// then reading one byte at a time with pdb_stream_read_u8
+PDB_API pdb_uint
+pdb_stream_read_str(pdb_stream *stream, void *dst, pdb_uint dst_max)
+{
+  pdb_context *pdb = pdb_stream_get_pdb_context(stream);
+  pdb_uint read_size = 0;
+  if (pdb) {
+    U8 *ptr = (U8 *)pdb->file_data;
+    pdb_uint str_off = stream->off;
+    pdb_uint str_size = 0;
+
+    while (stream->off < stream->size) {
+      pdb_uint c = 0;
+      pdb_uint n;
+      pdb_uint scan_size;
+
+      SYMS_ASSERT(stream->page_read_hi <= pdb->file_size);
+      SYMS_ASSERT(stream->page_read_lo <= stream->page_read_hi);
+      SYMS_ASSERT(stream->page_read_hi - stream->page_read_lo <= stream->page_size);
+
+      // scan page forward and check for null and if not found
+      // move to next page
+      for (n=stream->page_read_lo; n<stream->page_read_hi; ++n) {
+        c = ptr[n];
+        if (c == '\0') {
+          ++n;
+          break;
+        }
+      }
+      scan_size = n - stream->page_read_lo;
+      str_size += scan_size;
+      if (c == '\0') {
+        break;
+      }
+      if (!pdb_stream_skip(stream, scan_size)) {
+        break;
+      }
+    }
+
+    // seek to string start
+    if (!pdb_stream_seek(stream, str_off)) {
+      SYMS_INVALID_CODE_PATH;
+    }
+
+    if (dst && dst_max) {
+      pdb_uint to_read = SYMS_MIN(str_size, dst_max);
+      pdb_uint slack = str_size - to_read;
+      read_size = pdb_stream_read(stream, dst, to_read); 
+      pdb_stream_skip(stream, slack);
+    } else {
+      read_size = str_size;
+    }
+  }
+  return read_size;
+}
+
+PDB_API pdb_uint
+pdb_stream_strlen(pdb_stream *stream)
+{
+  pdb_stream str = *stream;
+  pdb_uint result = pdb_stream_read_str(&str, 0, 0);
+  return result;
+}
+
+PDB_API syms_bool
+pdb_stream_can_read_bytes(pdb_stream *stream, U32 num_bytes)
+{
+  syms_bool result = stream->off + num_bytes <= stream->size;
+  return result;
+}
+
+PDB_API syms_bool
+pdb_stream_align(pdb_stream *stream, U32 align)
+{
+  syms_bool is_aligned = syms_true;
+  U32 mask = align - 1;
+  SYMS_ASSERT(align > 0);
+
+  if (stream->page_read_lo & mask) {
+    U32 align_off = align - (stream->page_read_lo & mask); 
+    is_aligned = pdb_stream_skip(stream, align_off);
+  }
+
+  return is_aligned;
+}
+
+PDB_API U32
+pdb_stream_get_abs_off(pdb_stream *stream)
+{
+  return (stream->off_at_subset + stream->off);
+}
+
+PDB_API pdb_stream
+pdb_stream_subset(const pdb_stream *stream, U32 off, U32 size)
+{
+  pdb_stream subset;
+
+  if (off + size <= stream->size) {
+    subset = *stream;
+    if (pdb_stream_seek(&subset, off)) {
+      subset.off = 0;
+      subset.off_at_subset += off;
+      subset.size = size;
+    } else {
+      pdb_stream_init_null(&subset);
+    }
+  } else {
+    pdb_stream_init_null(&subset);
+  }
+
+  return subset;
+}
+
+PDB_API pdb_pointer
+pdb_pointer_bake(pdb_stream *bake_stream, U32 bake_size)
+{
+  pdb_pointer result = pdb_pointer_bake_null();
+  if (bake_size <= PDB_POINTER_PAGE_MAX*bake_stream->pdb->page_size) {
+    pdb_stream stream = *bake_stream;
+    pdb_context *pdb = stream.pdb;
+    if (pdb) {
+      U32 page_size = pdb->page_size;
+      U32 i = 0;
+      result.mode = PDB_POINTER_MODE_PAGES;
+      while (bake_size > 0) {
+        U32 data_size;
+
+        if (i >= PDB_POINTER_PAGE_MAX) {
+          break;
+        }
+        data_size = syms_trunc_u32(stream.page_read_hi - stream.page_read_lo);
+        if (data_size <= page_size) {
+          data_size = SYMS_MIN(bake_size, data_size);
+        } else {
+          /* NOTE(nick): An error occurred in seek routine that failed and set
+           * page bounds too high and fitting read into 32-bits
+           * wont be possible. Break out and return whatever was read. */
+          SYMS_ASSERT_FAILURE("invalid page bounds");
+          break;
+        }
+        result.u.pages.offs[i] = stream.page_read_lo;
+        result.u.pages.size[i] = data_size;
+        if (!pdb_stream_skip(&stream, data_size)) {
+          break;
+        }
+        i += 1;
+        bake_size -= data_size;
+      }
+    }
+  } else {
+    result.mode = PDB_POINTER_MODE_STREAM;
+    result.u.stream.sn = bake_stream->sn;
+    result.u.stream.off = bake_stream->off_at_subset + bake_stream->off;
+    result.u.stream.size = bake_size;
+  }
+  return result;
+}
+PDB_API U32
+pdb_pointer_read(pdb_context *pdb, const pdb_pointer *bytes, U32 off, void *bf, U32 bf_max)
+{
+  U32 result = 0;
+  switch (bytes->mode) {
+  case PDB_POINTER_MODE_NULL: {
+    result = 0;
+  } break;
+  case PDB_POINTER_MODE_RAW: {
+    if (off <= bytes->u.raw.size) {
+      result = SYMS_MIN(bf_max, (bytes->u.raw.size - off));
+      syms_memcpy(bf, (u8 *)bytes->u.raw.data + off, result);
+    }
+  } break;
+  case PDB_POINTER_MODE_PAGES: {
+    U32 i = 0;
+    U32 curr_off = 0;
+    for (;;) {
+      if (i >= PDB_POINTER_PAGE_MAX) {
+        break;
+      }
+      if (curr_off + bytes->u.pages.size[i] > off) {
+        break;
+      }
+      curr_off += bytes->u.pages.size[i++];
+    }
+    off -= curr_off;
+
+    {
+      /* NOTE(nick): A more light-weight read from MSF stream, however
+       * it is limited on amount of bytes that it fetches.
+       * Size of a fetch equals to PDB_POINTER_PAGE_MAX * page_size.
+       **/
+
+      U8 *bf_write = (U8 *)bf;
+      U32 bf_size = bf_max;
+
+      while (bf_size > 0) {
+        U32 copy_size;
+        const void *src_data;
+
+        if (i >= PDB_POINTER_PAGE_MAX) {
+          break;
+        }
+        if (bytes->u.pages.offs[i] == 0 || bytes->u.pages.size[i] == 0) {
+          break;
+        }
+
+        copy_size = SYMS_MIN(bytes->u.pages.size[i], bf_size); 
+
+        src_data = (const void *)((const u8 *)pdb->file_data + bytes->u.pages.offs[i] + off);
+        syms_memcpy((void *)bf_write, src_data, copy_size);
+
+        bf_write += copy_size;
+        bf_size -= copy_size;
+
+        i += 1;
+        off = 0;
+      }
+
+      result = (bf_max - bf_size);
+    }
+  } break;
+  case PDB_POINTER_MODE_STREAM: {
+    pdb_stream stream;
+    if (pdb_stream_init(pdb, bytes->u.stream.sn, &stream)) {
+      u32 read_size = SYMS_MIN(bytes->u.stream.size, bf_max); 
+      stream = pdb_stream_subset(&stream, bytes->u.stream.off, bytes->u.stream.size);
+      result = pdb_stream_read_(&stream, off, bf, read_size);
+    }
+  } break;
+  }
+
+  return result;
+}
+PDB_API pdb_uint
+pdb_pointer_strlen(pdb_context *pdb, pdb_pointer *ptr, u32 offset)
+{
+  u32 i = offset;
+  pdb_uint k = 0;
+
+  for (;;) {
+    u8 b;
+
+    if (k >= PDB_STRLEN_MAX) {
+      break;
+    }
+    b = pdb_pointer_read_u08(pdb, ptr, i);
+    i += 1;
+    k += 1;
+    if (b == '\0') {
+      break;
+    }
+  }
+  return k;
+}
+PDB_API pdb_uint 
+pdb_pointer_read_str(pdb_context *pdb, pdb_pointer *ptr, u32 offset, u8 *buf, pdb_uint buf_size)
+{
+  u32 i = offset;
+  pdb_uint k = 0;
+  for (;;) {
+    if (k >= buf_size) {
+      break;
+    }
+    buf[k] = pdb_pointer_read_u08(pdb, ptr, i);
+    if (buf[k] == '\0') {
+      break;
+    }
+    i += 1;
+    k += 1;
+  }
+  return k;
+}
+PDB_API pdb_pointer
+pdb_pointer_bake_null(void)
+{
+  pdb_pointer result;
+  syms_memset(&result, 0, sizeof(result));
+  return result;
+}
+PDB_API pdb_pointer
+pdb_pointer_bake_buffer(const void *buffer, pdb_uint buffer_size)
+{
+  pdb_pointer result = pdb_pointer_bake_null();
+  result.mode       = PDB_POINTER_MODE_RAW;
+  result.u.raw.data = buffer;
+  result.u.raw.size = buffer_size;
+  return result;
+}
+PDB_API pdb_pointer
+pdb_pointer_bake_str(struct SymsString str)
+{
+  return pdb_pointer_bake_buffer(str.data, (pdb_uint)str.len);
+}
+PDB_API pdb_pointer
+pdb_pointer_bake_stream_str(pdb_stream *msf)
+{
+  pdb_stream stream = *msf;
+  U32 str_size = pdb_stream_read_str(&stream, 0, 0);
+  if (str_size > 0) {
+    /* TODO(nick): Size includes null, remove dumb -1 */
+    stream = *msf;
+    return pdb_pointer_bake(msf, str_size - 1);
+  }
+  return pdb_pointer_bake_null();
+}
+PDB_API pdb_pointer
+pdb_pointer_bake_sn(pdb_context *pdb, pdb_sn sn, U32 off, U32 len)
+{
+  pdb_stream stream;
+  if (pdb_stream_init(pdb, sn, &stream)) {
+    if (pdb_stream_seek(&stream, off)) {
+      return pdb_pointer_bake(&stream, len);
+    }
+  } 
+  SYMS_ASSERT_FAILURE("cannot bake data from stream");
+  return pdb_pointer_bake_null();
+}
+PDB_API U32
+pdb_pointer_get_size(pdb_pointer *pointer)
+{
+  U32 size = 0;
+  switch (pointer->mode) {
+  case PDB_POINTER_MODE_NULL: {
+    /* ignore */
+  } break;
+  case PDB_POINTER_MODE_RAW: {
+    size = pointer->u.raw.size;
+  } break;
+  case PDB_POINTER_MODE_PAGES: {
+    U32 i;
+    for (i = 0; i < PDB_POINTER_PAGE_MAX; ++i) {
+      size += pointer->u.pages.size[i];
+    }
+  } break;
+  case PDB_POINTER_MODE_STREAM: {
+    size = pointer->u.stream.size;
+  } break;
+  }
+  return size;
+}
+PDB_API U32
+pdb_pointer_read_u32(pdb_context *pdb, pdb_pointer *bytes, U32 off)
+{
+  U32 result = 0;
+  if (pdb_pointer_read(pdb, bytes, off, &result, sizeof(result)) != sizeof(result)) {
+#if defined(SYMS_PARANOID) 
+    SYMS_ASSERT_FAILURE("invalid read");
+#endif
+  }
+  return result;
+}
+PDB_API U16
+pdb_pointer_read_u16(pdb_context *pdb, pdb_pointer *bytes, U32 off)
+{
+  U16 result = 0;
+  if (pdb_pointer_read(pdb, bytes, off, &result, sizeof(result)) != sizeof(result)) {
+#if defined(SYMS_PARANOID) 
+    SYMS_ASSERT_FAILURE("invalid read");
+#endif
+  }
+  return result;
+}
+PDB_API U8
+pdb_pointer_read_u08(pdb_context *pdb, pdb_pointer *bytes, U32 off)
+{
+  U8 result = 0;
+  if (pdb_pointer_read(pdb, bytes, off, &result, sizeof(result)) != sizeof(result)) {
+#if defined(SYMS_PARANOID)
+    SYMS_ASSERT_FAILURE("invalid read");
+#endif
+  }
+  return result;
+}
+PDB_API U32
+pdb_pointer_read_utf8(pdb_context *pdb, pdb_pointer *bytes, U32 off, U32 *codepoint_out)
+{
+  /* TODO(nick): IMPLEMENT UTF8 */
+  *codepoint_out = pdb_pointer_read_u08(pdb, bytes, off);
+  if (*codepoint_out == 0)
+    return 0;
+  return 1;
+}
+PDB_API syms_bool
+pdb_pointer_cmp(pdb_context *pdb, pdb_pointer *bytes_a, pdb_pointer *bytes_b)
+{
+  syms_bool is_equal = syms_false;
+  U32 a_size = pdb_pointer_get_size(bytes_a);
+  U32 b_size = pdb_pointer_get_size(bytes_b);
+  if (a_size == b_size) {
+    U32 read_off = 0;
+    while (read_off < a_size) {
+      U8 chunk_a[32];
+      U8 chunk_b[32];
+      U32 read_size = SYMS_MIN((a_size - read_off), sizeof(chunk_a));
+      U32 read_a = pdb_pointer_read(pdb, bytes_a, read_off, chunk_a, read_size);
+      U32 read_b = pdb_pointer_read(pdb, bytes_b, read_off, chunk_b, read_size);
+      if (read_a == read_b) {
+        if (syms_memcmp(chunk_a, chunk_b, read_a) != 0) {
+          break;
+        }
+      }
+      read_off += read_a;
+    }
+    is_equal = (read_off == a_size);
+  }
+  return is_equal;
+}
+PDB_API syms_bool
+pdb_pointer_strcmp_(pdb_context *pdb, pdb_pointer *bytes_a, pdb_pointer *bytes_b, pdb_strcmp_flags_e cmp_flags)
+{
+  syms_bool cmp = syms_false;
+  if (cmp_flags & PDB_STRCMP_FLAG_NOCASE) {
+    U32 len_a = pdb_pointer_get_size(bytes_a);
+    U32 len_b = pdb_pointer_get_size(bytes_b);
+    U32 i;
+
+    if (len_a == len_b) {
+      i = 0;
+      while (i < len_a) {
+        U32 a, b;
+        U32 cp_len_a = pdb_pointer_read_utf8(pdb, bytes_a, i, &a);
+        U32 cp_len_b = pdb_pointer_read_utf8(pdb, bytes_b, i, &b);
+
+        if (cp_len_a != cp_len_b) break;
+        i += cp_len_a;
+
+        a = pdb_trunc_uint(syms_lowercase(a));
+        b = pdb_trunc_uint(syms_lowercase(b));
+
+        if (a != b) break;
+      }
+
+      cmp = (i == len_a);
+    }
+  } else {
+    cmp = pdb_pointer_cmp(pdb, bytes_a, bytes_b);
+  }
+  return cmp;
+}
+
+PDB_API syms_bool
+pdb_strcmp_stream_(SymsString str, pdb_stream *stream, pdb_strcmp_flags_e cmp_flags)
+{
+  U32 rewind_off = stream->off;
+  pdb_pointer bytes_a = pdb_pointer_bake_str(str);
+  pdb_pointer bytes_b = pdb_pointer_bake_stream_str(stream);
+  syms_bool cmp = pdb_pointer_strcmp_(stream->pdb, &bytes_a, &bytes_b, cmp_flags);
+  if (~cmp_flags & PDB_STRCMP_FLAG_NOCASE && !pdb_stream_seek(stream, rewind_off)) {
+    SYMS_INVALID_CODE_PATH;
+  }
+  return cmp;
+}
+PDB_API syms_bool
+pdb_stream_strcmp_stream_(pdb_stream *stream_a, pdb_stream *stream_b, pdb_strcmp_flags_e cmp_flags)
+{
+  pdb_pointer bytes_a = pdb_pointer_bake_stream_str(stream_a);
+  pdb_pointer bytes_b = pdb_pointer_bake_stream_str(stream_b);
+  syms_bool cmp;
+  SYMS_ASSERT(stream_a->pdb == stream_b->pdb);
+  cmp = pdb_pointer_strcmp_(stream_a->pdb, &bytes_a, &bytes_b, cmp_flags);
+  return cmp;
+}
+PDB_API syms_bool
+pdb_stream_strcmp_pointer_(pdb_stream *stream, pdb_pointer *pointer, pdb_strcmp_flags_e cmp_flags)
+{
+  pdb_pointer bytes_a = pdb_pointer_bake_stream_str(stream);
+  syms_bool cmp = pdb_pointer_strcmp_(stream->pdb, &bytes_a, pointer, cmp_flags);
+  return cmp;
+}
+PDB_API syms_bool
+pdb_strcmp_pointer_(pdb_context *pdb, SymsString str_a, pdb_pointer *str_b, pdb_strcmp_flags_e cmp_flags)
+{
+  pdb_pointer str_pointer_a = pdb_pointer_bake_str(str_a);
+  syms_bool cmp = pdb_pointer_strcmp_(pdb, &str_pointer_a, str_b, cmp_flags);
+  return cmp;
+}
+
+/* -------------------------------------------------------------------------------- 
+ *                                  PDB
+ * -------------------------------------------------------------------------------- */
+
+PDB_API U32
+pdb_hashV1_bytes(const void *start, U32 cb, U32 mod)
+{
+  pdb_pointer pointer = pdb_pointer_bake_buffer(start, cb);
+  U32 hash = pdb_hashV1_pointer(NULL, &pointer, mod);
+  return hash;
+}
+PDB_API U32
+pdb_hashV1_stream(pdb_stream *stream, U32 size, U32 mod)
+{
+  pdb_pointer pointer = pdb_pointer_bake(stream, size);
+  U32 hash = pdb_hashV1_pointer(stream->pdb, &pointer, mod);
+  return hash;
+}
+PDB_API U32
+pdb_hashV1_pointer(pdb_context *pdb, pdb_pointer *bytes, U32 mod)
+{
+  U32 hash = 0;
+#if 0
+  U32 bytes_len = pdb_pointer_get_size(bytes);
+  U32 num_u32_reads = bytes_len/sizeof(U32);
+  U32 i;
+
+  U32 read_off = 0;
+
+  U32 remainder_size = bytes_len % 4;
+
+  for (i = 0; i < num_u32_reads; ++i) {
+    hash ^= pdb_pointer_read_u32(pdb, bytes, read_off);
+    read_off += sizeof(U32);
+  }
+
+  if (bytes_len - read_off >= sizeof(U16)) {
+    hash ^= pdb_pointer_read_u16(pdb, bytes, read_off);
+    read_off += sizeof(U16);
+  }
+
+  if (bytes_len - read_off == sizeof(U8)) {
+    hash ^= pdb_pointer_read_u08(pdb, bytes, read_off);
+  }
+#else
+  U32 read_off = 0;
+  U32 read_max = pdb_pointer_get_size(bytes);
+  U32 max_off_for_u32 = read_max >> 2; /* Same as a div by 4 */
+  U32 advance_for_u32 = max_off_for_u32 & 7;
+
+  // TODO(nick): I tried appending [[fallthrough]] to each case, but then clang 
+  // says that this syntax is not compatible with C++98. Explicitly suppressing 
+  // warning is a crude solution and if you know a better one you are welcome to 
+  // replace it.
+#ifdef __clang__
+  #pragma clang diagnostic push
+  #pragma clang diagnostic ignored "-Wimplicit-fallthrough"
+#endif
+  switch (advance_for_u32) {
+    do {
+        advance_for_u32 = 8;
+        hash ^= pdb_pointer_read_u32(pdb, bytes, (read_off + 7)*sizeof(U32));
+    case 7: hash ^= pdb_pointer_read_u32(pdb, bytes, (read_off + 6)*sizeof(U32)); 
+    case 6: hash ^= pdb_pointer_read_u32(pdb, bytes, (read_off + 5)*sizeof(U32));
+    case 5: hash ^= pdb_pointer_read_u32(pdb, bytes, (read_off + 4)*sizeof(U32));
+    case 4: hash ^= pdb_pointer_read_u32(pdb, bytes, (read_off + 3)*sizeof(U32));
+    case 3: hash ^= pdb_pointer_read_u32(pdb, bytes, (read_off + 2)*sizeof(U32));
+    case 2: hash ^= pdb_pointer_read_u32(pdb, bytes, (read_off + 1)*sizeof(U32));
+    case 1: hash ^= pdb_pointer_read_u32(pdb, bytes, (read_off + 0)*sizeof(U32));
+    case 0: ;
+
+      read_off += advance_for_u32;
+    } while (read_off < max_off_for_u32);
+  }
+#ifdef __clang__
+  #pragma clang diagnostic pop
+#endif
+
+  read_off *= sizeof(U32);
+  /* hash possible odd word */
+  if (read_max & 2) {
+    hash ^= pdb_pointer_read_u16(pdb, bytes, read_off);
+    read_off += sizeof(U16);
+  }
+  /* hash possible odd byte */
+  if (read_max & 1) {
+    hash ^= pdb_pointer_read_u08(pdb, bytes, read_off);
+    read_off += sizeof(U8);
+  }
+  SYMS_ASSERT(read_off == read_max);
+#endif
+
+  hash |= 0x20202020;
+  hash ^= (hash >> 11);
+  hash ^= (hash >> 16);
+  if (mod > 0) {
+    hash %= mod;
+  }
+  return hash;
+}
+
+PDB_API U32
+pdb_calc_size_for_types(pdb_tm_header *tm_header)
+{
+  U32 ti_count = tm_header->ti_hi - tm_header->ti_lo;
+  U32 result = tm_header->hash_bucket_count*sizeof(((pdb_tm *)0)->buckets[0]) +
+           ti_count*sizeof(pdb_tm_bucket) +
+           ti_count*sizeof(((pdb_tm *)0)->ti_offsets[0]) + 
+           128; /* For alignment */
+  return result;
+}
+
+PDB_API syms_bool
+pdb_tm_init(pdb_tm *tm, pdb_context *pdb, pdb_default_stream sn, SymsArena *arena)
+{
+  pdb_tm_header *tm_header = &tm->header;
+  pdb_stream tm_data;
+  syms_bool is_inited = syms_false;
+
+  tm->pdb         = pdb;
+  tm->sn          = sn;
+  tm->ti_offsets  = 0;
+  tm->buckets     = 0;
+
+  SYMS_ASSERT(sn == PDB_DEFAULT_STREAM_TPI || sn == PDB_DEFAULT_STREAM_IPI);
+
+  if (!pdb_stream_init(pdb, sn, &tm_data))
+    return syms_false;
+  if (!pdb_stream_read_struct(&tm_data, tm_header, pdb_tm_header))
+    return syms_false;
+
+  is_inited = syms_false;
+  switch (tm_header->version) {
+  // unsupported versions
+  case PDB_INTV_VC2:
+  case PDB_TM_IMPV40:
+  case PDB_TM_IMPV41:
+  case PDB_TM_IMPV50_INTERIM:
+  case PDB_TM_IMPV50:
+  case PDB_TM_IMPV70: break;
+
+  // supported versions
+  case PDB_TM_IMPV80: {
+    pdb_stream hash_data;
+    pdb_uint ti_count;
+    pdb_uint ti;
+
+    if (tm_header->hash_key_size != 4 || tm_header->hash_bucket_count < 0x1000 || tm_header->hash_bucket_count >= 0x40000)
+      break;
+
+    ti_count = tm_header->ti_hi - tm_header->ti_lo;
+    if (ti_count == 0) {
+      is_inited = syms_true; // this is an empty type map
+      break;
+    }
+
+    if (!pdb_stream_init(pdb, tm_header->hash_sn, &hash_data))
+      break;
+    if (!pdb_stream_seek(&hash_data, tm_header->hash_vals.off))
+      break;
+    if (!pdb_stream_can_read_bytes(&hash_data, tm_header->hash_vals.cb))
+      break;
+
+    // allocate slots for type index offsets
+    tm->ti_offsets = syms_arena_push_array(arena, U32, ti_count);
+    // allocate buckets
+    tm->buckets = syms_arena_push_array(arena, pdb_tm_bucket *, tm_header->hash_bucket_count);
+    if (!tm->ti_offsets || !tm->buckets)
+      break; // allocation failed
+
+    // zero out memory
+    syms_memset(tm->ti_offsets, 0, sizeof(tm->ti_offsets[0])*ti_count);
+    syms_memset(tm->buckets, 0, sizeof(pdb_tm_bucket *)*tm_header->hash_bucket_count);
+
+    // populate an internally-chained hash table; each bucket has
+    // a linked list of type indices.
+    for (ti = tm_header->ti_lo; ti < tm_header->ti_hi; ++ti) {
+      U32 bucket_index = PDB_UINT_MAX;
+      pdb_tm_bucket *bucket;
+      pdb_stream_read_u32(&hash_data, &bucket_index);
+      if (bucket_index >= tm->header.hash_bucket_count) 
+        continue;
+      bucket = syms_arena_push_struct(arena, pdb_tm_bucket);
+      bucket->ti = ti;
+      bucket->next = tm->buckets[bucket_index];
+      tm->buckets[bucket_index] = bucket;
+    }
+
+    if (tm_header->hash_adj.cb > 0 && tm_header->hash_adj.cb != ~0u) {
+      pdb_stream hash_stream;
+      pdb_uint i;
+
+      pdb_uint num_present_and_deleted = 0;
+      pdb_uint table_size = 0;
+      pdb_uint n_bits_present = 0;
+      pdb_uint n_bits_deleted = 0;
+      pdb_uint bits_off_present;
+      pdb_uint adj_num = 0;
+      pdb_uint adj_off = 0;
+      pdb_stream strtable;
+
+      if (!pdb_stream_init(pdb, tm_header->hash_sn, &hash_stream))
+        break;
+
+      i = (tm_header->ti_off.cb / sizeof(pdb_off_cb))*sizeof(pdb_off_cb);
+      if (i != tm_header->ti_off.cb)
+        break;
+      if (!pdb_stream_seek(&hash_stream, (U32)tm_header->hash_adj.off))
+        break;
+      if (!pdb_stream_can_read_bytes(&hash_stream, sizeof(U32)*3))
+        break;
+
+      pdb_stream_read_u32(&hash_stream, &num_present_and_deleted);
+      pdb_stream_read_u32(&hash_stream, &table_size);
+      pdb_stream_read_u32(&hash_stream, &n_bits_present);
+      bits_off_present = hash_stream.off;
+      if (!pdb_stream_skip(&hash_stream, n_bits_present*sizeof(U32)))
+        break;
+      pdb_stream_read_u32(&hash_stream, &n_bits_deleted);
+      // U32 bits_off_deleted = hash_stream.off;
+      adj_num = 0;
+      adj_off = hash_stream.off + n_bits_deleted*sizeof(U32);
+
+      strtable = pdb_get_strtable(pdb);
+      for (i = 0; i < table_size; ++i) {
+        U32 bits_off = bits_off_present + (i >> 5)*sizeof(U32);
+        U32 bits;
+
+        if (!pdb_stream_seek(&hash_stream, bits_off))  continue;
+        if (!pdb_stream_read_u32(&hash_stream, &bits)) continue;
+
+        if (bits & (1 << (i & 31))) {
+          pdb_ti read_ti;
+          pdb_pointer name;
+          pdb_uint nameoff;
+          pdb_uint bucket_index;
+
+          pdb_tm_bucket *prev_bucket;
+          pdb_tm_bucket *bucket;
+
+          if (!pdb_stream_seek(&hash_stream, adj_off + adj_num*sizeof(U32)*2)) continue;
+          if (!pdb_stream_read_u32(&hash_stream, &nameoff))   continue;
+          if (!pdb_stream_read_u32(&hash_stream, &read_ti))        continue;
+
+          name = pdb_pointer_bake_stream_str(&strtable);
+          bucket_index = pdb_hashV1_pointer(pdb, &name, tm->header.hash_bucket_count);
+
+          prev_bucket = 0;
+          bucket      = tm->buckets[bucket_index];
+          while (bucket) {
+            if (bucket->ti == read_ti) {
+              if (prev_bucket) {
+                prev_bucket->next = bucket->next;
+                bucket->next = tm->buckets[bucket_index];
+                tm->buckets[bucket_index] = bucket;
+              }
+              break;
+            }
+
+            prev_bucket = bucket;
+            bucket = bucket->next;
+          }
+
+          adj_num += 1;
+        }
+      }
+
+      if (PDB_STREAM_READ_OR_SEEK_FAILED(hash_stream.flags))
+        break;
+    }
+    is_inited = syms_true;
+  } break;
+    
+  default: {
+#if defined(SYMS_PARANOID)
+    SYMS_ASSERT_NO_SUPPORT;
+#endif
+  } break;
+  }
+
+  return is_inited;
+}
+
+PDB_API syms_bool
+pdb_tm_offset_for_ti(pdb_tm *tm, pdb_ti ti, u32 *ti_off_out)
+{
+  /* For some reason PDB people could not export all offsets for the ti,
+   * which means that we have look into the table to see where the closest
+   * range for the specified ti starts and resolve offset from there, and
+   * at the end write the offset into the table for future re-use. */
+
+  pdb_context *pdb = tm->pdb;
+  pdb_tm_header *tm_header = &tm->header;
+  pdb_ti ti_index;
+  pdb_stream hash_stream;
+
+  u32 i, c;
+  s32 min, max, mid;
+
+  pdb_ti_off curr_tioff;
+
+  u32 blk_min, blk_max;
+  pdb_stream tm_data;
+
+  if (tm->ti_offsets == 0 || tm->buckets == 0)
+    return syms_false;
+  if (ti < tm_header->ti_lo || ti >= tm_header->ti_hi)
+    return syms_false;
+
+  ti_index = ti - tm_header->ti_lo;
+  if (tm->ti_offsets[ti_index] != 0) {
+    /* Code below was already executed and it saved the offset,
+     * return it, and call it a day. */
+    *ti_off_out = tm->ti_offsets[ti_index];
+    return syms_true;
+  }
+  
+  if (!pdb_stream_init(pdb, tm_header->hash_sn, &hash_stream))
+    return syms_false;
+  if (!pdb_stream_seek(&hash_stream, tm_header->ti_off.off))
+    return syms_false;
+  c = tm->header.ti_off.cb/sizeof(pdb_ti_off);
+  if (c == 0)
+    return syms_false; // empty type map
+  if (!pdb_stream_read_struct(&hash_stream, &curr_tioff, pdb_ti_off))
+    return syms_false;
+
+  i = c;
+  min = 0;
+  max = (s32)c - 1;
+  mid = 0;
+  while (min <= max) {
+    pdb_uint info_off;
+
+    mid = min + (max - min) / 2;
+
+    // seek and read type index with offset
+    info_off = tm_header->ti_off.off + mid*sizeof(pdb_ti_off);
+    if (!pdb_stream_seek(&hash_stream, info_off))
+      return syms_false;
+    if (!pdb_stream_read_struct(&hash_stream, &curr_tioff, pdb_ti_off))
+      return syms_false;
+
+    if (curr_tioff.ti > ti) {
+      max = mid - 1;
+    } else if (curr_tioff.ti < ti) {
+      min = mid + 1;
+    } else {
+      i = mid;
+      break;
+    }
+  }
+  // binary search landed on item that is off by one, adjust item index
+  // so it points to correct slot.
+  if (ti < curr_tioff.ti && mid > 0) {
+    pdb_uint info_off;
+
+    i = mid - 1;
+    info_off = tm_header->ti_off.off + i*sizeof(pdb_ti_off);
+    if (!pdb_stream_seek(&hash_stream, info_off))
+      return syms_false;
+    if (!pdb_stream_read_struct(&hash_stream, &curr_tioff, pdb_ti_off))
+      return syms_false;
+  }
+
+  if (curr_tioff.ti < tm_header->ti_lo || curr_tioff.ti >= tm_header->ti_hi)
+    return syms_false;
+
+  blk_min = curr_tioff.ti;
+  blk_max = 0;
+  // u32 blk_len;
+  if (i < (c - 1)) {
+    pdb_ti_off next_tioff;
+    if (!pdb_stream_read_struct(&hash_stream, &next_tioff, pdb_ti_off))
+      return syms_false;
+    blk_max = next_tioff.ti;
+    // blk_len = next_tioff.off - curr_tioff.off;
+  } else {
+    blk_max = tm_header->ti_hi;
+    // u32 off = curr_tioff.off + tm_header->header_size;
+    // blk_len = tm_header->types_size + tm_header->header_size - off;
+  }
+  SYMS_ASSERT(ti >= blk_min);
+  SYMS_ASSERT(ti < blk_max);
+
+  if (!pdb_stream_init(pdb, tm->sn, &tm_data))
+    return syms_false;
+  if (!pdb_stream_seek(&tm_data, curr_tioff.off + tm_header->header_size))
+    return syms_false;
+  for (i = blk_min; i < blk_max; ++i) {
+    pdb_symrec symrec;
+
+    tm->ti_offsets[i - tm_header->ti_lo] = tm_data.off;
+    if (!pdb_stream_read_symrec(&tm_data, &symrec))
+      break;
+    if (!pdb_stream_seek(&tm_data, symrec.end))
+      break;
+  }
+
+  *ti_off_out = tm->ti_offsets[ti_index];
+  return tm->ti_offsets[ti_index] != 0;
+}
+
+PDB_API syms_bool
+pdb_tm_get_itype_offset(pdb_context *pdb, pdb_cv_itype itype, U32 *itype_off_out)
+{
+  return pdb_tm_offset_for_ti(&pdb->tpi, itype, itype_off_out);
+}
+
+PDB_API syms_bool
+pdb_tm_get_itemid_offset(pdb_context *pdb, pdb_cv_itemid itemid, U32 *itemid_off)
+{
+  return pdb_tm_offset_for_ti(&pdb->ipi, itemid, itemid_off);
+}
+
+PDB_API syms_bool
+pdb_find_udt_srcline(pdb_context *pdb, pdb_cv_itype lookup_itype, pdb_udt_srcline *srcline_out)
+{
+  pdb_tm *tm = &pdb->ipi;
+  pdb_stream stream;
+  u32 bucket_index;
+  pdb_tm_bucket *tm_bucket;
+
+  if (!pdb_stream_init(pdb, tm->sn, &stream))
+    return syms_false;
+
+  bucket_index = pdb_hashV1_bytes(&lookup_itype, sizeof(lookup_itype), tm->header.hash_bucket_count);
+  for (tm_bucket = tm->buckets[bucket_index]; 
+      tm_bucket != 0; 
+      tm_bucket = tm_bucket->next) {
+    u32 ti_off;
+    pdb_symrec symrec;
+    pdb_cv_itype itype = tm_bucket->ti;
+
+    if (!pdb_tm_offset_for_ti(tm, itype, &ti_off))  continue;
+    if (!pdb_stream_seek(&stream, ti_off))          continue;
+    if (!pdb_stream_read_symrec(&stream, &symrec))  continue;
+
+    switch (symrec.type) {
+    case PDB_LF_UDT_MOD_SRC_LINE: {
+      pdb_lf_modsrcline lf;
+      pdb_stream strtable;
+
+      lf.udt_itype = PDB_INVALID_ITYPE;
+      lf.src = PDB_INVALID_ITYPE;
+      lf.ln = 0;
+      lf.mod = SYMS_UINT16_MAX;
+      pdb_stream_read_struct(&stream, &lf, pdb_lf_modsrcline);
+      if (lf.udt_itype != lookup_itype) {
+        break;
+      }
+
+      strtable = pdb_get_strtable(pdb);
+      pdb_stream_seek(&strtable, lf.src);
+      srcline_out->file = pdb_pointer_bake_stream_str(&strtable);
+      srcline_out->ln = (pdb_uint)lf.ln;
+      srcline_out->mod = (pdb_imod)lf.mod;
+      return syms_true;
+    }
+
+    case PDB_LF_UDT_SRC_LINE: {
+      pdb_lf_srcline lf;
+      pdb_type file_name;
+
+      lf.udt_itype = PDB_INVALID_ITYPE;
+      lf.src = PDB_INVALID_ITYPE;
+      lf.ln = 0;
+      pdb_stream_read_struct(&stream, &lf, pdb_lf_srcline);
+      if (lf.udt_itype != lookup_itype) {
+        break;
+      }
+
+      srcline_out->file = pdb_pointer_bake_null();
+      srcline_out->ln = (pdb_uint)lf.ln;
+      srcline_out->mod = PDB_CV_INVALID_IMOD;
+      if (pdb_infer_itemid(pdb, lf.src, &file_name)) {
+        if (file_name.kind == PDB_TYPE_STRINGID) {
+          srcline_out->file = file_name.u.stringid.data;
+        } else {
+          SYMS_ASSERT_FAILURE_PARANOID("expected PDB_TYPE_STRING_ID");
+        }
+      } else {
+        SYMS_ASSERT_FAILURE_PARANOID("unable to resolve file name");
+      }
+
+      return syms_true;
+    }
+    }
+  }
+  
+  return syms_false;
+}
+
+PDB_API syms_bool
+pdb_tm_find_ti(pdb_tm *tm, pdb_pointer *name, pdb_ti *ti_out)
+{
+  pdb_context *pdb = tm->pdb;
+  pdb_stream stream;
+  u32 bucket_index;
+  pdb_tm_bucket *tm_bucket;
+
+  if (!pdb_stream_init(pdb, tm->sn, &stream))
+    return syms_false;
+
+  bucket_index = pdb_hashV1_pointer(pdb, name, tm->header.hash_bucket_count);
+  for (tm_bucket = tm->buckets[bucket_index]; 
+      tm_bucket != 0; 
+      tm_bucket = tm_bucket->next) {
+#if 1
+    u32 ti_off;
+    pdb_symrec symrec;
+    u32 prop;
+
+    pdb_cv_itype itype = tm_bucket->ti;
+    syms_bool keep_resolving;
+    do {
+      keep_resolving = syms_false;
+
+      if (!pdb_tm_offset_for_ti(tm, itype, &ti_off))  continue;
+      if (!pdb_stream_seek(&stream, ti_off))          continue;
+      if (!pdb_stream_read_symrec(&stream, &symrec))  continue;
+
+      prop = PDB_CV_PROP_FWDREF;
+      switch (symrec.type) {
+      case PDB_LF_INTERFACE:
+      case PDB_LF_CLASS:
+      case PDB_LF_STRUCTURE: {
+        pdb_lf_class udt;
+        if (!pdb_stream_read_struct(&stream, &udt, pdb_lf_class))
+          SYMS_INVALID_CODE_PATH;
+        if (!pdb_stream_read_numeric(&stream, 0))
+          SYMS_INVALID_CODE_PATH;
+        prop = udt.prop;
+      } break;
+      case PDB_LF_ENUM: {
+        pdb_lf_enum udt;
+        if (!pdb_stream_read_struct(&stream, &udt, pdb_lf_enum))
+          SYMS_INVALID_CODE_PATH;
+        prop = udt.prop;
+      } break;
+      case PDB_LF_UNION: {
+        pdb_lf_union udt;
+        if (!pdb_stream_read_struct(&stream, &udt, pdb_lf_union))
+          SYMS_INVALID_CODE_PATH;
+        if (!pdb_stream_read_numeric(&stream, 0))
+          SYMS_INVALID_CODE_PATH;
+        prop = udt.prop;
+      } break;
+      case PDB_LF_POINTER: {
+        pdb_lf_ptr ptr;
+        if (!pdb_stream_read_struct(&stream, &ptr, pdb_lf_ptr))
+          SYMS_INVALID_CODE_PATH;
+        itype = ptr.itype;
+        keep_resolving = syms_true;
+      } break;
+      case PDB_LF_MODIFIER: {
+        pdb_lf_modifier modifier;
+        if (!pdb_stream_read_struct(&stream, &modifier, pdb_lf_modifier))
+          SYMS_INVALID_CODE_PATH;
+        itype = modifier.itype;
+        keep_resolving = syms_true;
+      } break;
+      case PDB_LF_MFUNCTION: {
+        pdb_lf_mfunc mfunc;
+        if (!pdb_stream_read_struct(&stream, &mfunc, pdb_lf_mfunc))
+          SYMS_INVALID_CODE_PATH;
+        itype = mfunc.classtype;
+        keep_resolving = syms_true;
+      } break;
+      case PDB_LF_CLASSPTR2:
+      case PDB_LF_CLASSPTR: {
+        pdb_lf_classptr lf;
+        u32 size;
+
+        if (!pdb_stream_read_struct(&stream, &lf, pdb_lf_classptr))
+          SYMS_INVALID_CODE_PATH;
+        pdb_stream_read_numeric_u32(&stream, &size);
+        prop = lf.prop;
+      } break;
+      default: break;
+      }
+
+      if (~prop & PDB_CV_PROP_FWDREF) {
+        if (pdb_stream_strcmp_pointer(&stream, name)) { 
+          *ti_out = tm_bucket->ti;
+          return syms_true;
+        }
+      }
+#else
+      pdb_type ti_type;
+      if (pdb_tm_infer_ti(tm, tm_bucket->ti, &ti_type)) {
+        if (pdb_pointer_strcmp(&ti_type.name, name)) {
+          *ti_out = tm_bucket->ti;
+          return syms_true;
+        }
+      }
+#endif
+    } while (keep_resolving);
+  }
+  
+  if (pdb->globals_array_num > 0) {
+    pdb_gsi_hr *gsi_bucket;
+
+    if (!pdb_stream_init(pdb, pdb->dbi.symrec_sn, &stream)) {
+      return syms_false;
+    }
+
+    bucket_index = pdb_hashV1_pointer(pdb, name, pdb->globals_array_num);
+    gsi_bucket = pdb->globals_array[bucket_index];
+    for (gsi_bucket = pdb->globals_array[bucket_index]; 
+        gsi_bucket != 0; 
+        gsi_bucket = gsi_bucket->next) {
+      pdb_symrec symrec;
+      if (!pdb_stream_seek(&stream, gsi_bucket->off)) {
+        SYMS_ASSERT_CORRUPTED_STREAM;
+        continue;
+      }
+      if (!pdb_stream_read_symrec(&stream, &symrec)) {
+        SYMS_ASSERT_CORRUPTED_STREAM;
+        continue;
+      }
+      if (symrec.type == PDB_CV_SYM_CONSTANT) {
+        pdb_stream_skip(&stream, sizeof(pdb_cv_constsym));
+      }
+      if (symrec.type == PDB_CV_SYM_UDT) {
+        pdb_cv_udtsym sym;
+        if (pdb_stream_read_struct(&stream, &sym, pdb_cv_udtsym)) {
+          if (pdb_stream_strcmp_pointer(&stream, name)) {
+            *ti_out = sym.itype;
+            return syms_true;
+          }
+        }
+      } 
+    }
+  }
+
+  return syms_false;
+}
+
+PDB_API const char *
+pdb_ver_to_str(pdb_context *pdb)
+{
+  switch (pdb->ver) {
+  case PDB_VER_VC50:      return "VC50";
+  case PDB_VER_VC4:       return "VC4";
+  case PDB_VER_VC2:       return "VC2";
+  case PDB_VER_VC98:      return "VC98";
+  case PDB_VER_VC70:      return "VC70";
+  case PDB_VER_VC70_DEP:  return "VC70_DEP";
+  case PDB_VER_VC80:      return "VC80";
+  case PDB_VER_VC140:     return "VC140";
+  case PDB_VER_VC110:     return "VC110";
+  }
+  return "";
+}
+
+PDB_API const char *
+pdb_dbi_ver_to_str(pdb_context *pdb)
+{
+  switch (pdb->dbi.header.version) {
+  case PDB_DBI_VER_41: return "41";
+  case PDB_DBI_VER_50: return "50";
+  case PDB_DBI_VER_60: return "60";
+  case PDB_DBI_VER_70: return "70";
+  case PDB_DBI_VER_110: return "110";
+  }
+  return "";
+}
+
+PDB_API const char *
+pdb_basic_itype_to_str(pdb_basic_type itype)
+{
+  switch (itype) {
+  case PDB_BASIC_TYPE_NOTYPE:    return "none"; 			    break;
+  case PDB_BASIC_TYPE_ABS:       return "abs"; 				    break;
+  case PDB_BASIC_TYPE_SEGMENT:   return "segment"; 			  break;
+  case PDB_BASIC_TYPE_VOID:      return "void"; 			    break;
+  case PDB_BASIC_TYPE_CURRENCY:  return "currency"; 			break;
+  case PDB_BASIC_TYPE_NBASICSTR: return "nbasicstr";			break;
+  case PDB_BASIC_TYPE_FBASICSTR: return "fbasicstr";			break;
+  case PDB_BASIC_TYPE_NOTTRANS:  return "nottrans"; 			break;
+  case PDB_BASIC_TYPE_BIT: 	     return "bit"; 				    break;
+  case PDB_BASIC_TYPE_PASCHAR:   return "PascalChar"; 	  break;
+  case PDB_BASIC_TYPE_BOOL32FF:  return "bool32ff"; 			break;
+  case PDB_BASIC_TYPE_HRESULT:   return "HRESULT"; 			  break;
+  case PDB_BASIC_TYPE_RCHAR:     return "rchar"; 			    break;
+  case PDB_BASIC_TYPE_WCHAR:     return "wchar_t"; 			  break;
+  case PDB_BASIC_TYPE_CHAR8:     return "char8";          break;
+  case PDB_BASIC_TYPE_CHAR16:    return "char16"; 		    break;
+  case PDB_BASIC_TYPE_CHAR32: 	 return "char32"; 		    break;
+
+  case PDB_BASIC_TYPE_INT1:      return "int8_t"; 			break;
+  case PDB_BASIC_TYPE_INT2:      return "int16_t"; 			break;
+  case PDB_BASIC_TYPE_INT4:      return "int32_t"; 			break;
+  case PDB_BASIC_TYPE_INT8:      return "int64_t"; 			break;
+  case PDB_BASIC_TYPE_INT16: 	   return "int128_t"; 		break;
+
+  case PDB_BASIC_TYPE_UINT1:     return "uint8_t"; 			  break;
+  case PDB_BASIC_TYPE_UINT2:     return "uint16_t";  			break;
+  case PDB_BASIC_TYPE_UINT4:     return "uint32_t";  			break;
+  case PDB_BASIC_TYPE_UINT8:     return "uint64_t";  			break;
+  case PDB_BASIC_TYPE_UINT16: 	 return "uint128_t"; 			break;
+
+  case PDB_BASIC_TYPE_CHAR: 	 return "char"; 			break;
+  case PDB_BASIC_TYPE_SHORT: 	 return "short"; 			break;
+  case PDB_BASIC_TYPE_LONG:    return "long"; 			break;
+  case PDB_BASIC_TYPE_QUAD:    return "long long"; 	break;
+
+  case PDB_BASIC_TYPE_UCHAR:     return "unsigned char"; 		  break;
+  case PDB_BASIC_TYPE_USHORT:    return "unsigned short";  		break;
+  case PDB_BASIC_TYPE_UQUAD:     return "unsigned long long";	break;
+  case PDB_BASIC_TYPE_ULONG:     return "unsigned long"; 		  break;
+
+  case PDB_BASIC_TYPE_OCT: 	 return "octal"; 			        break;
+  case PDB_BASIC_TYPE_UOCT:  return "unsigned octal"; 		break;
+
+  case PDB_BASIC_TYPE_REAL16: 	 return "real16"; 			break;
+  case PDB_BASIC_TYPE_REAL32: 	 return "real32"; 			break;
+  case PDB_BASIC_TYPE_REAL32PP:  return "real32 (partial precision)";   break;
+  case PDB_BASIC_TYPE_REAL48:    return "real48"; 			break;
+  case PDB_BASIC_TYPE_REAL64:    return "real64"; 			break;
+  case PDB_BASIC_TYPE_REAL80:    return "real80"; 			break;
+  case PDB_BASIC_TYPE_REAL128:   return "real128"; 			break;
+  case PDB_BASIC_TYPE_CPLX32:    return "complex32"; 		break;
+  case PDB_BASIC_TYPE_CPLX64:    return "complex64"; 		break;
+  case PDB_BASIC_TYPE_CPLX128:   return "complex128"; 	break;
+  case PDB_BASIC_TYPE_BOOL08:    return "bool8"; 		    break;
+  case PDB_BASIC_TYPE_BOOL16:    return "bool16";		    break;
+  case PDB_BASIC_TYPE_BOOL32:    return "bool32";		    break;
+  case PDB_BASIC_TYPE_BOOL64:    return "bool64";		    break;
+  case PDB_BASIC_TYPE_PTR:       return "void *";		    break;
+  }
+  return 0;
+}
+
+PDB_API syms_bool
+pdb_type_from_name(pdb_context *pdb, const char *name, U32 name_len, pdb_type *type_out)
+{
+  SymsString name_str = syms_string_init(name, name_len);
+  pdb_pointer name_ref = pdb_pointer_bake_str(name_str);
+  pdb_ti ti;
+  pdb_cv_itype itype;
+
+  if (pdb_tm_find_ti(&pdb->tpi, &name_ref, &ti))
+    return pdb_tm_infer_ti(&pdb->tpi, ti, type_out);
+
+  if (pdb_tm_find_ti(&pdb->ipi, &name_ref, &ti))
+    return pdb_tm_infer_ti(&pdb->ipi, ti, type_out);
+
+  for (itype = PDB_BASIC_TYPE_NOTYPE; itype < PDB_BASIC_TYPE_MAX; ++itype) {
+    const char *itype_str = pdb_basic_itype_to_str(itype);
+    if (itype_str) {
+      SymsString str = syms_string_init_lit(itype_str);
+      if (syms_string_cmp(str, name_str))
+        return pdb_tm_infer_ti(&pdb->tpi, itype, type_out);
+    }
+  }
+
+  return syms_false;
+}
+
+PDB_API syms_bool
+pdb_infer_basic_itype(pdb_context *pdb, U32 itype, pdb_type *type)
+{
+  U32 itype_size = PDB_BASIC_TYPE_SIZE_MASK(itype); 
+  U32 itype_kind = PDB_BASIC_TYPE_KIND_MASK(itype);
+
+  switch (itype_size) {
+  case PDB_BASIC_TYPE_SIZE_VALUE: {
+    switch (itype_kind) {
+    case PDB_BASIC_TYPE_VOID: {
+      type->size = 0;
+      type->kind = PDB_TYPE_VOID;
+    } break;
+
+    case PDB_BASIC_TYPE_HRESULT: {
+      type->size = 4;
+      type->kind = PDB_TYPE_VOID;
+    } break;
+
+    case PDB_BASIC_TYPE_RCHAR:
+    case PDB_BASIC_TYPE_CHAR: {
+      type->size = sizeof(char);
+      type->kind = PDB_TYPE_CHAR;
+    } break;
+
+    case PDB_BASIC_TYPE_UCHAR: {
+      type->size = sizeof(unsigned char);
+      type->kind = PDB_TYPE_UCHAR;
+    } break;
+
+    case PDB_BASIC_TYPE_WCHAR: {
+      type->size = sizeof(wchar_t);
+      type->kind = PDB_TYPE_WCHAR;
+    } break;
+
+    case PDB_BASIC_TYPE_BOOL08:
+    case PDB_BASIC_TYPE_CHAR8:
+    case PDB_BASIC_TYPE_INT1: {
+      type->size = sizeof(S8);
+      type->kind = PDB_TYPE_INT8;
+    } break;
+
+    case PDB_BASIC_TYPE_BOOL16:
+    case PDB_BASIC_TYPE_CHAR16:
+    case PDB_BASIC_TYPE_SHORT:
+    case PDB_BASIC_TYPE_INT2: {
+      type->size = sizeof(S16);
+      type->kind = PDB_TYPE_INT16;
+    } break;
+
+    case PDB_BASIC_TYPE_BOOL32:
+    case PDB_BASIC_TYPE_CHAR32:
+    case PDB_BASIC_TYPE_INT4: {
+      type->size = sizeof(S32);
+      type->kind = PDB_TYPE_INT32;
+    } break;
+
+    case PDB_BASIC_TYPE_BOOL64:
+    case PDB_BASIC_TYPE_QUAD:
+    case PDB_BASIC_TYPE_INT8: {
+      type->size = sizeof(S64);
+      type->kind = PDB_TYPE_INT64;
+    } break;
+
+    case PDB_BASIC_TYPE_OCT: 
+    case PDB_BASIC_TYPE_INT16: {
+      type->size = sizeof(S64)*2;
+      type->kind = PDB_TYPE_INT128;
+    } break;
+
+    case PDB_BASIC_TYPE_UINT1: {
+      type->size = sizeof(U8);
+      type->kind = PDB_TYPE_UINT8;
+    } break;
+
+    case PDB_BASIC_TYPE_USHORT:
+    case PDB_BASIC_TYPE_UINT2: {
+      type->size = sizeof(U16);
+      type->kind = PDB_TYPE_UINT16;
+    } break;
+
+    case PDB_BASIC_TYPE_LONG: {
+      switch (pdb->dbi.machine_type) {
+      case SYMS_NT_FILE_HEADER_MACHINE_X64: {
+        type->size = sizeof(s64);
+        type->kind = PDB_TYPE_INT64;
+      } break;
+      case SYMS_NT_FILE_HEADER_MACHINE_X86: {
+        type->size = sizeof(s32);
+        type->kind = PDB_TYPE_INT32;
+      } break;
+      default: {
+        type->size = 0;
+        type->kind = PDB_TYPE_NULL;
+      } break;
+      }
+    } break;
+
+    case PDB_BASIC_TYPE_ULONG: {
+      switch (pdb->dbi.machine_type) {
+      case SYMS_NT_FILE_HEADER_MACHINE_X64: {
+        type->size = sizeof(u64);
+        type->kind = PDB_TYPE_UINT64;
+      } break;
+      case SYMS_NT_FILE_HEADER_MACHINE_X86: {
+        type->size = sizeof(u32);
+        type->kind = PDB_TYPE_UINT32;
+      } break;
+      default: {
+        type->size = 0;
+        type->kind = PDB_TYPE_NULL;
+      } break;
+      }
+    } break;
+
+    case PDB_BASIC_TYPE_UINT4: {
+      type->size = sizeof(U32);
+      type->kind = PDB_TYPE_UINT32;
+    } break;
+
+    case PDB_BASIC_TYPE_UQUAD:
+    case PDB_BASIC_TYPE_UINT8: {
+      type->size = sizeof(U64);
+      type->kind = PDB_TYPE_UINT64;
+    } break;
+
+    case PDB_BASIC_TYPE_UOCT:
+    case PDB_BASIC_TYPE_UINT16: {
+      type->size = sizeof(U64)*2;
+      type->kind = PDB_TYPE_UINT128;
+    } break;
+
+    case PDB_BASIC_TYPE_REAL16: {
+      type->size = 2;
+      type->kind = PDB_TYPE_REAL16;
+    } break;
+
+    case PDB_BASIC_TYPE_REAL32: {
+      type->size = 4;
+      type->kind = PDB_TYPE_REAL32;
+    } break;
+
+    case PDB_BASIC_TYPE_REAL64: {
+      type->size = 8;
+      type->kind = PDB_TYPE_REAL64;
+    } break;
+
+    case PDB_BASIC_TYPE_REAL32PP: {
+      type->size = 4;
+      type->kind = PDB_TYPE_REAL32PP;
+    } break;
+
+    case PDB_BASIC_TYPE_REAL80: {
+      type->size = 10;
+      type->kind = PDB_TYPE_REAL80;
+    } break;
+
+    case PDB_BASIC_TYPE_REAL128: {
+      type->size = 16;
+      type->kind = PDB_TYPE_REAL128;
+    } break;
+
+    case PDB_BASIC_TYPE_CPLX32: {
+      type->size = 4;
+      type->kind = PDB_TYPE_COMPLEX32;
+    } break;
+
+    case PDB_BASIC_TYPE_CPLX64: {
+      type->size = 8;
+      type->kind = PDB_TYPE_COMPLEX64;
+    } break;
+
+    case PDB_BASIC_TYPE_CPLX80: {
+      type->size = 10;
+      type->kind = PDB_TYPE_COMPLEX80;
+    } break;
+
+    case PDB_BASIC_TYPE_CPLX128: {
+      type->size = 16;
+      type->kind = PDB_TYPE_COMPLEX128;
+    } break;
+
+    case PDB_BASIC_TYPE_NOTYPE: {
+      type->size = 0;
+      type->kind = PDB_TYPE_NULL;
+    } break;
+
+    default:
+    case PDB_BASIC_TYPE_ABS:
+    case PDB_BASIC_TYPE_SEGMENT:
+    case PDB_BASIC_TYPE_NBASICSTR:
+    case PDB_BASIC_TYPE_CURRENCY: 
+    case PDB_BASIC_TYPE_FBASICSTR:
+    case PDB_BASIC_TYPE_BIT:
+    case PDB_BASIC_TYPE_PASCHAR:
+    case PDB_BASIC_TYPE_BOOL32FF: {
+      type->size = 0;
+      type->kind = PDB_TYPE_NULL;
+      SYMS_ASSERT_FAILURE_PARANOID("encountered unsupported types");
+    } break;
+    }
+  } break;
+
+  case PDB_BASIC_TYPE_SIZE_16BIT: {
+    type->size = 2;
+    type->kind = PDB_TYPE_PTR;
+    type->next_cv_itype = itype_kind;
+  } break;
+
+  case PDB_BASIC_TYPE_SIZE_FAR_16BIT: {
+    /* TODO: Do we need a special type for this? I think this is obsolete and nobody uses it.
+     * Same goes for the PDB_BASIC_TYPE_SIZE_HUGE_16BIT and PDB_BASIC_TYPE_SIZE_16_32BIT */
+    type->size = 2;
+    type->kind = PDB_TYPE_PTR;
+    type->next_cv_itype = itype_kind;
+  } break;
+
+  case PDB_BASIC_TYPE_SIZE_HUGE_16BIT: {
+    type->size = 2;
+    type->kind = PDB_TYPE_PTR;
+    type->next_cv_itype = itype_kind;
+  } break;
+
+  case PDB_BASIC_TYPE_SIZE_32BIT: {
+    type->size = 4;
+    type->kind = PDB_TYPE_PTR;
+    type->next_cv_itype = itype_kind;
+  } break;
+
+  case PDB_BASIC_TYPE_SIZE_16_32BIT: {
+    type->size = 4;
+    type->kind = PDB_TYPE_PTR;
+    type->next_cv_itype = itype_kind;
+  } break;
+
+  case PDB_BASIC_TYPE_SIZE_64BIT: {
+    type->size = 8;
+    type->kind = PDB_TYPE_PTR;
+    type->next_cv_itype = itype_kind;
+  } break;
+  }
+
+  type->name = pdb_pointer_bake_str(pdb->basic_typenames_array[itype_kind]);
+
+  return syms_true;
+}
+
+PDB_API syms_bool
+pdb_tm_infer_ti(pdb_tm *tm, pdb_ti ti, pdb_type *type)
+{
+  union generic_leaf {
+    struct pdb_lf_class     udt_struct;
+    struct pdb_lf_enum      udt_enum;
+    struct pdb_lf_array     udt_array;
+    struct pdb_lf_union     udt_union;
+    struct pdb_lf_classptr  udt_classptr;
+    struct pdb_lf_modifier  mod;
+    struct pdb_lf_ptr       ptr;
+    struct pdb_lf_bitfield  bit;
+    struct pdb_lf_proc      proc;
+    struct pdb_lf_mproc     mproc;
+    struct pdb_lf_arglist   arglist;
+    struct pdb_lf_func_id   func_id;
+    struct pdb_lf_mfunc_id  mfunc_id;
+    struct pdb_lf_string_id string_id;
+    struct pdb_lf_vtshape   vtshape;
+    struct pdb_lf_vftable   vftable;
+    struct pdb_lf_label     label;
+  };
+
+  pdb_context *pdb = tm->pdb;
+  pdb_stream stream;
+
+  syms_bool type_done = syms_false;
+  syms_bool fetch_more = syms_false;
+
+  union generic_leaf lf;
+
+  if (!pdb_stream_init(pdb, tm->sn, &stream)) 
+    return syms_false;
+
+  syms_memzero(&lf, sizeof(lf));
+  syms_memzero(type, sizeof(*type));
+
+  do {
+    U32 typeoff;
+    pdb_symrec symrec;
+
+    fetch_more = syms_false;
+
+    if (ti == PDB_ITYPE_VARIADIC) {
+      type->kind = PDB_TYPE_VARIADIC;
+      type_done = syms_true;
+      break;
+    }
+    if (ti < tm->header.ti_lo) {
+      /* Type index of a basic type */
+      type_done = pdb_infer_basic_itype(pdb, ti, type);
+      break;
+    }
+
+    if (ti >= tm->header.ti_hi)
+      break; /* unknown type index */
+
+    /* Map type index into a stream offset. */
+    if (!pdb_tm_offset_for_ti(tm, ti, &typeoff)) 
+      break; 
+
+    /* Seek and read a symbol record. */
+    pdb_stream_seek(&stream, typeoff);
+    pdb_stream_read_symrec(&stream, &symrec);
+    if (PDB_STREAM_READ_OR_SEEK_FAILED(stream.flags)) break;
+
+    switch (symrec.type) {
+    default: {
+      type_done  = syms_false;
+      fetch_more = syms_false;
+      SYMS_ASSERT_FAILURE_PARANOID("undefined type");
+    } break;
+
+    case PDB_LF_VTSHAPE: {
+      if (pdb_stream_read_struct(&stream, &lf.vtshape, pdb_lf_vtshape)) {
+        pdb_uint data_size = (pdb_uint)((f32)(lf.vtshape.count * sizeof(u8)) / 2.0f + 0.5f);
+        type->kind = PDB_TYPE_VTSHAPE;
+        type->u.vtshape.count = lf.vtshape.count;
+        type->u.vtshape.ptr = pdb_pointer_bake(&stream, data_size);
+        type_done = syms_true;
+      }
+      SYMS_ASSERT_PARANOID(!fetch_more);
+    } break;
+
+    case PDB_LF_LABEL: {
+      if (pdb_stream_read_struct(&stream, &lf.label, pdb_lf_label)) {
+        type->kind = PDB_TYPE_LABEL;
+        type->u.label.mode = (pdb_cv_ptrmode)lf.label.mode;
+        type_done = syms_true;
+      }
+    } break;
+
+    case PDB_LF_VFTABLE: {
+      if (pdb_stream_read_struct(&stream, &lf.vftable, pdb_lf_vftable)) {
+        u32 method_names_len;
+
+        type->u.vftable.owner_itype = lf.vftable.owner_itype;
+        type->u.vftable.base_table_itype = lf.vftable.base_table_itype;
+        type->u.vftable.offset_in_object_layout = lf.vftable.offset_in_object_layout;
+        type->u.vftable.name = pdb_pointer_bake_stream_str(&stream);
+        pdb_stream_skip(&stream, pdb_pointer_get_size(&type->u.vftable.name));
+        method_names_len = lf.vftable.names_len - pdb_pointer_get_size(&type->u.vftable.name);
+        type->u.vftable.method_names = pdb_pointer_bake(&stream, method_names_len);
+        type_done = syms_true;
+      }
+    } break;
+
+    case PDB_LF_METHODLIST: {
+      type->kind = PDB_TYPE_METHODLIST;
+      type->u.methodlist.block = pdb_pointer_bake(&stream, symrec.size);
+      type_done = syms_true;
+    } break;
+
+    case PDB_LF_FIELDLIST: {
+      type->kind = PDB_TYPE_FIELDLIST;
+      type->u.fieldlist.data = pdb_pointer_bake(&stream, symrec.size);
+      type_done = syms_true;
+    } break;
+
+    case PDB_LF_FUNC_ID: {
+      if (pdb_stream_read_struct(&stream, &lf.func_id, pdb_lf_func_id)) {
+        type->kind = PDB_TYPE_FUNCID;
+        type->u.funcid.itype = lf.func_id.itype;
+        type->u.funcid.name = pdb_pointer_bake_stream_str(&stream);
+        type_done = syms_true;
+      }
+    } break;
+
+    case PDB_LF_MFUNC_ID: {
+      if (pdb_stream_read_struct(&stream, &lf.mfunc_id, pdb_lf_mfunc_id)) {
+        type->kind = PDB_TYPE_MFUNCID;
+        type->u.mfuncid.name = pdb_pointer_bake_stream_str(&stream);
+        type->u.mfuncid.parent_itype = lf.mfunc_id.parent_itype;
+        type->u.mfuncid.itype = lf.mfunc_id.itype;
+        type_done = syms_true;
+      }
+    } break;
+
+    case PDB_LF_STRING_ID: {
+      if (pdb_stream_read_struct(&stream, &lf.string_id, pdb_lf_string_id)) {
+        type->kind = PDB_TYPE_STRINGID;
+        type->u.stringid.data = pdb_pointer_bake_stream_str(&stream);
+        type->u.stringid.sub_string = lf.string_id.id;
+        type_done = syms_true;
+      }
+    } break;
+
+    case PDB_LF_MODIFIER: {
+      if (pdb_stream_read_struct(&stream, &lf.mod, pdb_lf_modifier)) {
+        ti = lf.mod.itype;
+        if (lf.mod.attr & PDB_CV_MODIFIER_CONST)
+          type->attribs |= PDB_TYPE_ATTRIB_CONST;
+        if (lf.mod.attr & PDB_CV_MODIFIER_VOLATILE)
+          type->attribs |= PDB_TYPE_ATTRIB_VOLATILE;
+        if (lf.mod.attr & PDB_CV_MODIFIER_UNALIGNED)
+          type->attribs |= PDB_TYPE_ATTRIB_UNALIGNED;
+
+        fetch_more = syms_true;
+        SYMS_ASSERT(!type_done);
+      }
+    } break;
+
+    case PDB_LF_BITFIELD: {
+      if (pdb_stream_read_struct(&stream, &lf.bit, pdb_lf_bitfield)) {
+        type->kind = PDB_TYPE_BITFIELD;
+        type->u.bitfield.base_itype = lf.bit.itype;
+        type->u.bitfield.len = lf.bit.len;
+        type->u.bitfield.pos = lf.bit.pos;
+
+        SYMS_ASSERT(!fetch_more);
+        type_done = syms_true;
+      }
+    } break;
+
+    case PDB_LF_MFUNCTION: {
+      if (pdb_stream_read_struct(&stream, &lf.mproc, pdb_lf_mproc)) {
+        type->kind = PDB_TYPE_METHOD;
+        type->u.method.conv = lf.mproc.call_kind;
+        type->u.method.ret_itype = lf.mproc.ret_itype;
+        type->u.method.class_itype = lf.mproc.class_itype;
+        type->u.method.this_itype = lf.mproc.this_itype;
+        type->u.method.arg_itype = lf.mproc.arg_itype;
+        type->u.method.arg_count = lf.mproc.arg_count;
+
+        SYMS_ASSERT(!fetch_more);
+        type_done = syms_true;
+      }
+    } break;
+
+    case PDB_LF_PROCEDURE: {
+      if (pdb_stream_read_struct(&stream, &lf.proc, pdb_lf_proc)) {
+        type->kind = PDB_TYPE_PROC;
+        type->next_cv_itype = lf.proc.ret_itype;
+        type->u.proc.conv = lf.proc.call_kind;
+        type->u.proc.ret_itype = lf.proc.ret_itype;
+        type->u.proc.arg_itype = lf.proc.arg_itype;
+        type->u.proc.arg_count = lf.proc.arg_count;
+
+        SYMS_ASSERT(!fetch_more);
+        type_done = syms_true;
+      }
+    } break;
+
+    case PDB_LF_ARRAY: {
+      /* When an array is multi-dimensional(e.g vars[2][2][2][2][2])
+       * compiler writes every dimension as a single leaf that points to
+       * the next LF_ARRAY leaf. Field "entry_itype" of lf_array contains 
+       * index of next LF_ARRAY leaf. When no more dimensions are left
+       * "entry_itype" will point to the type of the array.
+       *                                               - Nikita, Aug 8 2017. */
+
+      if (pdb_stream_read_struct(&stream, &lf.udt_array, pdb_lf_array)) {
+        /* Contains the size of the array, which includes the sub dimensions if present. */
+        if (!pdb_stream_read_numeric_u32(&stream, &type->size))
+          break;
+
+        type->kind = PDB_TYPE_ARR;
+        type->next_cv_itype = lf.udt_array.entry_itype;
+
+        SYMS_ASSERT(!fetch_more);
+        type_done = syms_true;
+      }
+    } break;
+
+    case PDB_LF_POINTER: {
+      if (pdb_stream_read_struct(&stream, &lf.ptr, pdb_lf_ptr)) {
+        type->size = PDB_CV_PTR_ATTRIB_SIZE_MASK(lf.ptr.attr);
+        type->kind = PDB_TYPE_PTR;
+        type->next_cv_itype = lf.ptr.itype;
+        type->u.ptr.type = PDB_CV_PTR_ATTRIB_TYPE_MASK(lf.ptr.attr);
+        type->u.ptr.mode = PDB_CV_PTR_ATTRIB_MODE_MASK(lf.ptr.attr);
+        type->u.ptr.attr = lf.ptr.attr;
+
+        if (lf.ptr.attr & PDB_CV_PTR_ATTRIB_IS_VOLATILE)    type->attribs |= PDB_TYPE_ATTRIB_VOLATILE;
+        if (lf.ptr.attr & PDB_CV_PTR_ATTRIB_IS_CONST)       type->attribs |= PDB_TYPE_ATTRIB_CONST;
+        if (lf.ptr.attr & PDB_CV_PTR_ATTRIB_IS_UNALIGNED)   type->attribs |= PDB_TYPE_ATTRIB_UNALIGNED;
+        if (lf.ptr.attr & PDB_CV_PTR_ATTRIB_IS_RESTRICTED)  type->attribs |= PDB_TYPE_ATTRIB_RESTRICTED;
+        if (lf.ptr.attr & PDB_CV_PTR_ATTRIB_IS_LREF)        type->attribs |= PDB_TYPE_ATTRIB_LREF;
+        if (lf.ptr.attr & PDB_CV_PTR_ATTRIB_IS_RREF)        type->attribs |= PDB_TYPE_ATTRIB_RREF;
+
+        SYMS_ASSERT(!fetch_more);
+        type_done = syms_true;
+      }
+    } break;
+
+    case PDB_LF_ARGLIST: {
+      if (pdb_stream_read_struct(&stream, &lf.arglist, pdb_lf_arglist)) {
+        u32 bake_size = lf.arglist.count * sizeof(pdb_cv_itype);
+
+        type->kind             = PDB_TYPE_ARGLIST;
+        type->u.arglist.count  = lf.arglist.count;
+        type->u.arglist.itypes = pdb_pointer_bake(&stream, bake_size);
+
+        type_done = syms_true;
+      }
+    } break;
+
+#if 0
+    case 0x1609: {
+      void *buffer = malloc(symrec.size);
+      pdb_stream_read(&stream, buffer, symrec.size);
+      int x = 0;
+    } break;
+#endif
+
+    case PDB_LF_CLASSPTR2:
+    case PDB_LF_CLASSPTR:
+    case PDB_LF_UNION:
+    case PDB_LF_ENUM:
+    case PDB_LF_CLASS:
+    case PDB_LF_STRUCTURE: {
+      syms_bool fwdref = syms_false;
+
+      switch (symrec.type) {
+      case PDB_LF_CLASSPTR2:
+      case PDB_LF_CLASSPTR: {
+        pdb_stream_read_struct(&stream, &lf.udt_classptr, pdb_lf_classptr);
+        pdb_stream_read_numeric_u32(&stream, &type->size);
+
+        fwdref = (lf.udt_classptr.prop & PDB_CV_PROP_FWDREF) != 0;
+        if (fwdref) break;
+
+      write_classptr:;
+        type->kind = PDB_TYPE_STRUCT;
+        type->name = pdb_pointer_bake_null();
+        type->u.udt.field_itype = lf.udt_classptr.arglist_itype;
+        type->u.udt.field_count = lf.udt_classptr.unknown4;
+      } break;
+
+      case PDB_LF_UNION: {
+        pdb_stream_read_struct(&stream, &lf.udt_union, pdb_lf_union);
+        pdb_stream_read_numeric_u32(&stream, &type->size);
+
+        fwdref = (lf.udt_union.prop & PDB_CV_PROP_FWDREF) != 0;
+        if (fwdref) break;
+
+      write_union:;
+        type->kind = PDB_TYPE_UNION;
+        type->u.udt.field_itype = lf.udt_union.field;
+        type->u.udt.field_count = lf.udt_union.count;
+      } break;
+
+      case PDB_LF_CLASS:
+      case PDB_LF_STRUCTURE: {
+        pdb_stream_read_struct(&stream, &lf.udt_struct, pdb_lf_class);
+        pdb_stream_read_numeric_u32(&stream, &type->size);
+
+        fwdref = (lf.udt_struct.prop & PDB_CV_PROP_FWDREF) != 0;
+        if (fwdref) break;
+
+      write_struct:;
+        if (symrec.type == PDB_LF_CLASS) {
+          type->kind = PDB_TYPE_CLASS;
+        } else if (symrec.type == PDB_LF_STRUCTURE) {
+          type->kind = PDB_TYPE_STRUCT;
+        } else {
+          type->kind = PDB_TYPE_NULL;
+        }
+        type->u.udt.field_itype = lf.udt_struct.field;
+        type->u.udt.field_count = lf.udt_struct.count;
+      } break;
+
+      case PDB_LF_ENUM: {
+        pdb_stream_read_struct(&stream, &lf.udt_enum, pdb_lf_enum);
+
+        fwdref = (lf.udt_enum.prop & PDB_CV_PROP_FWDREF) != 0;
+        if (fwdref) break;
+
+      write_enum:;
+        type_done = pdb_infer_basic_itype(pdb, lf.udt_enum.itype, type);
+        if (!type_done) break;
+
+        type->kind = PDB_TYPE_ENUM;
+        type->u.udt.field_itype = lf.udt_enum.field;
+        type->u.udt.field_count = lf.udt_enum.count;
+        type->u.udt.base_itype = lf.udt_enum.itype;
+      } break;
+
+      default: return syms_false;
+      }
+      SYMS_ASSERT(!PDB_STREAM_READ_OR_SEEK_FAILED(stream.flags));
+
+      if (fwdref) {
+        pdb_pointer name = pdb_pointer_bake_stream_str(&stream);
+        pdb_ti new_ti;
+
+        type->attribs |= PDB_TYPE_ATTRIB_FWDREF;
+        if (!pdb_tm_find_ti(tm, &name, &new_ti)) {
+          type->name = name;
+          type_done = syms_true; // type is a forward reference
+          break;
+        }
+
+        if (ti == new_ti) {
+          SYMS_ASSERT(!fetch_more);
+
+          switch (symrec.type) {
+          case PDB_LF_ENUM: {
+            /* MSVC C/C++ compiler allows to forward declare enum 
+             * that does NOT exist! This creates an incomplete leaf
+             * type which has to be a special case. Visual Studio 
+             * in this case treats this entry as an int (4 bytes). */
+            lf.udt_enum.itype = PDB_BASIC_TYPE_INT4;
+            lf.udt_enum.prop &= ~PDB_CV_PROP_FWDREF;
+            goto write_enum;
+          }
+          case PDB_LF_CLASSPTR:   goto write_classptr;
+          case PDB_LF_STRUCTURE:  goto write_struct;
+          case PDB_LF_CLASS:      goto write_struct;
+          case PDB_LF_UNION:      goto write_union;
+          default: {
+            SYMS_ASSERT_FAILURE("Unexpected type for a fwd ref");
+            return syms_false;
+          }
+          }
+        }
+
+        ti = new_ti;
+        fetch_more = syms_true;
+        SYMS_ASSERT(!type_done);
+      } else {
+        type->name = pdb_pointer_bake_stream_str(&stream);
+        type_done = syms_true;
+        SYMS_ASSERT(!fetch_more);
+      }
+    } break;
+    }
+  } while (fetch_more);
+
+  if (type_done) {
+    type->cv_itype = ti;
+  }
+  return type_done;
+}
+
+PDB_API syms_bool
+pdb_infer_itemid(pdb_context *pdb, pdb_cv_itemid itemid, pdb_type *type_out)
+{
+  return pdb_tm_infer_ti(&pdb->ipi, itemid, type_out);
+}
+
+PDB_API syms_bool
+pdb_infer_itype(pdb_context *pdb, pdb_cv_itype itype, pdb_type *type_out)
+{
+  return pdb_tm_infer_ti(&pdb->tpi, itype, type_out);
+}
+
+PDB_API syms_bool
+pdb_type_it_init(pdb_type_it *it, pdb_context *pdb)
+{
+  it->pdb = pdb;
+  it->next_itype = pdb->tpi.header.ti_lo;
+  it->type_map_index = 0;
+  it->type_map[0] = &pdb->tpi;
+  it->type_map[1] = &pdb->ipi;
+  return syms_true;
+}
+
+PDB_API syms_bool
+pdb_type_it_next(pdb_type_it *it, pdb_cv_itype *itype_out)
+{
+  syms_bool result;
+  if (it->next_itype < it->type_map[0]->header.ti_hi)  {
+    *itype_out = it->next_itype;
+    it->next_itype += 1;
+    result = syms_true;
+  } else {
+    *itype_out = PDB_INVALID_ITYPE;
+    result = syms_false;
+  }
+  return result;
+}
+
+/* -------------------------------------------------------------------------------- */
+
+PDB_API syms_bool
+pdb_build_va(pdb_context *pdb, U32 isec, U32 off, SymsAddr *out_va)
+{
+  pdb_img_sec sec;
+  syms_bool is_result_valid = syms_false;
+
+  if (pdb_sec_from_index(pdb, isec, &sec)) {
+    *out_va = sec.rva + off;
+    is_result_valid = syms_true;
+  }
+
+  return is_result_valid;
+}
+
+PDB_API syms_bool
+pdb_build_sec_off(pdb_context *pdb, SymsAddr va, pdb_isec *sec, pdb_isec_umm *off)
+{
+  syms_bool result = syms_false;
+  SymsAddr rva = va;
+  pdb_isec i;
+
+  for (i = 0; i < pdb->dbi.secs_num; ++i) {
+    if (rva >= pdb->dbi.secs[i].rva && rva < pdb->dbi.secs[i].rva + pdb->dbi.secs[i].misc.virtual_size) {
+      *sec = i + 1;
+      *off = (pdb_isec_umm)(rva - pdb->dbi.secs[i].rva);
+      result = syms_true;
+      break;
+    }
+  }
+
+  return result;
+}
+
+PDB_API pdb_encoded_location
+pdb_encode_location_for_secoff(pdb_context *pdb, U32 isec, U32 off)
+{
+  pdb_encoded_location result;
+  pdb_img_sec sec;
+
+  if (pdb_sec_from_index(pdb, isec, &sec)) {
+    result.type = PDB_ENCODED_LOCATION_RVA;
+    result.flags = PDB_ENCODED_LOCATION_FLAG_NULL;
+    result.u.persist.va = sec.rva + off;
+  } else {
+    result.type = PDB_ENCODED_LOCATION_NULL;
+    result.flags = PDB_ENCODED_LOCATION_FLAG_NULL;
+  }
+  return result;
+}
+
+PDB_API pdb_encoded_location
+pdb_encode_location_for_datasym32(pdb_context *pdb, pdb_cv_datasym32 *datasym)
+{
+  return pdb_encode_location_for_secoff(pdb, datasym->sec, datasym->sec_off);
+} 
+PDB_API pdb_encoded_location
+pdb_encode_location_for_regrel(pdb_context *pdb, U32 regindex, U32 regoff)
+{
+  pdb_encoded_location result;
+
+  (void)pdb;
+
+  result.type = PDB_ENCODED_LOCATION_REGREL;
+  result.flags = PDB_ENCODED_LOCATION_FLAG_NULL;
+  result.u.regrel.reg_off = regoff;
+  result.u.regrel.reg_index = regindex;
+  return result;
+}
+
+PDB_API pdb_encoded_location
+pdb_encode_location_for_enreged(pdb_context *pdb, U32 reg_index)
+{
+  pdb_encoded_location result;
+
+  (void)pdb;
+
+  result.type = PDB_ENCODED_LOCATION_ENREGED;
+  result.flags = PDB_ENCODED_LOCATION_FLAG_NULL;
+  result.u.enreged.reg_index = reg_index;
+  return result;
+}
+
+PDB_API pdb_encoded_location
+pdb_encode_location_for_null(void)
+{
+  pdb_encoded_location result;
+  result.type = PDB_ENCODED_LOCATION_NULL;
+  result.flags = PDB_ENCODED_LOCATION_FLAG_NULL;
+  return result;
+}
+
+PDB_API syms_bool
+pdb_decode_location(pdb_encoded_location *encoded_loc,
+    SymsAddr orig_rebase,
+    SymsAddr rebase,
+    void *memread_ctx, pdb_memread_sig *memread,
+    void *regread_ctx, pdb_regread_sig *regread,
+    pdb_location *decoded_loc)
+{
+  syms_bool is_result_valid = syms_false;
+
+  (void)memread_ctx; (void)memread;
+
+  switch (encoded_loc->type) {
+  case PDB_ENCODED_LOCATION_ENREGED: {
+    if (regread) {
+      u32 read_size;
+
+      syms_memset(decoded_loc, 0, sizeof(*decoded_loc));
+      read_size = regread(
+          regread_ctx, 
+          encoded_loc->u.enreged.reg_index, 
+          (void *)&decoded_loc->u.implicit.data[0], 
+          sizeof(decoded_loc->u.implicit.data)); 
+      SYMS_ASSERT(read_size <= sizeof(decoded_loc->u.implicit.data));
+      decoded_loc->type = PDB_LOCATION_IMPLICIT;
+      decoded_loc->u.implicit.len = (U8)read_size;
+      is_result_valid = (decoded_loc->u.implicit.len > 0);
+    }
+  } break;
+  case PDB_ENCODED_LOCATION_REGREL: {
+    if (regread) {
+      pdb_regrel *regrel = &encoded_loc->u.regrel;
+      SymsUWord reg_value = 0;
+      U32 size = regread(regread_ctx, regrel->reg_index, &reg_value, sizeof(reg_value));
+      if (size > 0) {
+        decoded_loc->type = PDB_LOCATION_VA;
+        decoded_loc->u.va = reg_value + regrel->reg_off;
+        is_result_valid = syms_true;
+      }
+
+#if 0
+      pdb_context *pdb = encoded_loc->pdb;
+
+      pdb_type local_type;
+      if (pdb_infer_itype(pdb, pdb_local.itype, &local_type)) {
+        syms_bool is_aggregate_arg = pdb_local.u.regrel.is_arg &&
+          local_type.num_indirects == 0 &&
+          local_type.size > gpr_len &&
+          (local_type.kind == PDB_TYPE_STRUCT || local_type.kind == PDB_TYPE_UNION);
+
+        if (is_aggregate_arg) {
+          SYMS_ASSERT(ptr_len <= sizeof(va));
+          if (!memread(memread_context, va, &va, ptr_len))
+            break;
+        }
+      }
+#endif
+    }
+  } break;
+
+  case PDB_ENCODED_LOCATION_RVA: {
+    decoded_loc->type = PDB_LOCATION_VA;
+    decoded_loc->u.va = encoded_loc->u.persist.va + rebase;
+    is_result_valid = syms_true;
+  } break;
+
+  case PDB_ENCODED_LOCATION_VA: {
+    SYMS_ASSERT_PARANOID(encoded_loc->u.persist.va - orig_rebase);
+    decoded_loc->type = PDB_LOCATION_VA;
+    decoded_loc->u.va = (encoded_loc->u.persist.va - orig_rebase) - rebase;
+  } break;
+
+  case PDB_ENCODED_LOCATION_IMPLICIT: {
+    decoded_loc->type = PDB_LOCATION_IMPLICIT;
+    decoded_loc->u.implicit = encoded_loc->u.implicit;
+    is_result_valid = syms_true;
+  } break;
+
+  case PDB_ENCODED_LOCATION_NULL: break;
+  default: SYMS_INVALID_CODE_PATH; break;
+  }
+
+  return is_result_valid;
+}
+
+PDB_API syms_bool
+pdb_sec_it_init(pdb_context *pdb, pdb_sec_it *sec_it)
+{
+  return pdb_stream_init(pdb, pdb->dbi.dbg_streams[PDB_DBG_STREAM_SECTION_HEADER], &sec_it->stream);
+}
+
+PDB_API syms_bool
+pdb_sec_it_next(pdb_sec_it *sec_it, pdb_img_sec *sec)
+{
+  return pdb_stream_read_struct(&sec_it->stream, sec, pdb_img_sec);
+}
+
+PDB_API syms_bool
+pdb_sec_from_index(pdb_context *pdb, U32 index, pdb_img_sec *sec_out)
+{
+  syms_bool is_result_valid = syms_false;
+
+  if (index > 0 && index <= pdb->dbi.secs_num) {
+    index -= 1;
+
+    if (pdb->dbi.secs) {
+      *sec_out = pdb->dbi.secs[index];
+      is_result_valid = syms_true;
+    } else {
+      pdb_sec_it sec_it;
+
+      if (pdb_sec_it_init(pdb, &sec_it)) {
+        while (pdb_sec_it_next(&sec_it, sec_out)) {
+          if (index == 0)
+            break;
+          --index;
+        }
+        is_result_valid = syms_true;
+      }
+    }
+  }
+
+  return is_result_valid;
+}
+
+PDB_API syms_bool
+pdb_mod_it_init(pdb_mod_it *mod_it, pdb_context *pdb)
+{
+  syms_bool inited = syms_false;
+  pdb_stream stream;
+
+  if (pdb_stream_init(pdb, PDB_DEFAULT_STREAM_DBI, &stream)) {
+    mod_it->dbi_data = pdb_stream_subset(&stream, pdb->dbi.modinfo_off, pdb->dbi.modinfo_len);
+    if (mod_it->dbi_data.size > 0) {
+      mod_it->pdb  = pdb;
+      mod_it->imod = 0;
+      inited = syms_true;
+    }
+  }
+
+  return inited;
+}
+
+PDB_API syms_bool
+pdb_mod_it_next(pdb_mod_it *mod_it, pdb_mod *mod_out)
+{
+  pdb_mod_header header;
+  syms_bool was_read = syms_false;
+
+  pdb_stream_align(&mod_it->dbi_data, 4);
+  if (pdb_stream_read_struct(&mod_it->dbi_data, &header, pdb_mod_header)) {
+    u32 mod_name_length;
+
+    mod_out->pdb = mod_it->pdb;
+    mod_out->id = mod_it->imod;
+    mod_out->sn = header.sn;
+    mod_out->flags = 0;
+    mod_out->sec = header.sc.sec;
+    mod_out->sec_off = header.sc.sec_off;
+    mod_out->sec_len = header.sc.size;
+    mod_out->syms_size = header.symbol_bytes;
+    mod_out->c11_lines_size = header.c11_lines_size;
+    mod_out->c13_lines_size = header.c13_lines_size;
+
+    mod_out->name = pdb_pointer_bake_stream_str(&mod_it->dbi_data);
+    mod_name_length = pdb_pointer_get_size(&mod_out->name) + 1;
+    pdb_stream_skip(&mod_it->dbi_data, mod_name_length);
+
+    mod_out->name2 = pdb_pointer_bake_stream_str(&mod_it->dbi_data);
+    mod_name_length = pdb_pointer_get_size(&mod_out->name2) + 1;
+    pdb_stream_skip(&mod_it->dbi_data, mod_name_length);
+
+    mod_it->imod += 1;
+    was_read = syms_true;
+  }
+
+  return was_read;
+}
+
+PDB_API syms_bool
+pdb_mod_it_seek(pdb_mod_it *mod_it, pdb_imod imod)
+{
+  syms_bool was_moved = syms_false;
+
+  if (imod != mod_it->imod) {
+    if (imod < mod_it->pdb->dbi.mods_num) {
+      if (mod_it->pdb->dbi.mods) {
+        was_moved = pdb_stream_seek(&mod_it->dbi_data, mod_it->pdb->dbi.mods[imod]);
+      } else {
+        if (imod < mod_it->imod) {
+          pdb_mod_it_init(mod_it, mod_it->pdb);
+        }
+        while (imod > 0) {
+          pdb_mod dummy;
+          --imod;
+          if (!pdb_mod_it_next(mod_it, &dummy)) break;
+        }
+        was_moved = (imod == 0);
+      }
+    }
+  } else {
+    was_moved = syms_true;
+  }
+  return was_moved;
+}
+
+PDB_API syms_bool
+pdb_mod_get_debug_sec(pdb_mod *mod, pdb_mod_sec_type_e sec, pdb_stream *stream_out)
+{
+  syms_bool is_result_valid = syms_false;
+  pdb_stream stream;
+  if (pdb_stream_init(mod->pdb, mod->sn, &stream)) {
+    pdb_uint sig = 0;
+    pdb_uint size = 0;
+    pdb_uint off = PDB_UINT_MAX;
+    pdb_stream_read_u32(&stream, &sig);
+    switch (sec) {
+    case PDB_MOD_SEC_SYMS: {
+      switch (sig) {
+      case PDB_CV_SIG_C11: 
+      case PDB_CV_SIG_C13: {
+        off = sizeof(sig);
+        size = mod->syms_size;
+      } break;
+      case PDB_CV_SIG_C6:
+      case PDB_CV_SIG_C7: {
+        SYMS_ASSERT_NO_SUPPORT;
+      } break;
+      }
+    } break;
+    case PDB_MOD_SEC_LINES_C11: {
+      off = mod->syms_size;
+      size = mod->c11_lines_size;
+    } break;
+    case PDB_MOD_SEC_LINES_C13: {
+      off = mod->syms_size + mod->c11_lines_size;
+      size = mod->c13_lines_size;
+    } break;
+    case PDB_MOD_SEC_INLINE_LINES: {
+      SYMS_ASSERT_NO_SUPPORT;
+    } break;
+    }
+    *stream_out = pdb_stream_subset(&stream, off, size);
+    is_result_valid = (stream_out->size > 0);
+  }
+  return is_result_valid;
+}
+
+PDB_API syms_bool
+pdb_imod_from_isec(pdb_context *pdb, pdb_isec sec, pdb_isec_umm off, pdb_imod *imod_out)
+{
+  pdb_stream stream;
+  pdb_dbi_sc_ver ver;
+  U32 count, min, max, mid;
+
+  if (!pdb_stream_init(pdb, PDB_DEFAULT_STREAM_DBI, &stream)) {
+    return syms_false;
+  }
+
+  stream = pdb_stream_subset(&stream, pdb->dbi.seccon_off, pdb->dbi.seccon_len);
+  if (!pdb_stream_read(&stream, &ver, sizeof(ver))) {
+    return syms_false;
+  }
+
+  stream = pdb_stream_subset(&stream, sizeof(ver), stream.size - sizeof(ver));
+
+  if (ver == PDB_DBI_SC_VER_2) {
+    count = stream.size/sizeof(pdb_sc2);
+    if (count > 0) {
+      min = 0;
+      max = count - 1;
+      do {
+        pdb_sc2 sc;
+
+        sc.sec = 0;
+        sc.sec_off = 0;
+        sc.size = 0;
+
+        mid = (min + max)/2;
+        pdb_stream_seek(&stream, mid*sizeof(sc));
+        pdb_stream_read(&stream, &sc, sizeof(sc));
+
+        if (sec < sc.sec) {
+          max = mid - 1;
+        } else if (sec > sc.sec) {
+          min = mid + 1;
+        } else {
+          if (off < sc.sec_off) {
+            max = mid - 1;
+          } else if (off >= (sc.sec_off + sc.size)) {
+            min = mid + 1;
+          } else {
+            *imod_out = sc.imod;
+            return syms_true;
+          }
+        }
+      } while (min <= max && mid != 0);
+    }
+  } else if (ver == PDB_DBI_SC_VER_60) {
+    count = stream.size/sizeof(pdb_sc);
+    if (count > 0) {
+      min = 0;
+      max = count - 1;
+      do {
+        pdb_sc sc;
+
+        sc.sec = 0;
+        sc.sec_off = 0;
+        sc.size = 0;
+
+        mid = (min + max)/2;
+        pdb_stream_seek(&stream, mid*sizeof(sc));
+        pdb_stream_read(&stream, &sc, sizeof(sc));
+
+        if (sec < sc.sec) {
+          max = mid - 1;
+        } else if (sec > sc.sec) {
+          min = mid + 1;
+        } else {
+          if (off < sc.sec_off) {
+            max = mid - 1;
+          } else if (off >= (sc.sec_off + sc.size)) {
+            min = mid + 1;
+          } else {
+            *imod_out = sc.imod;
+            return syms_true;
+          }
+        }
+      } while (min <= max && mid != 0);
+    }
+  } else {
+    SYMS_INVALID_CODE_PATH;
+  }
+
+  return syms_false;
+}
+
+PDB_API syms_bool
+pdb_mod_init(pdb_mod *mod, pdb_context *pdb, pdb_imod imod)
+{
+  pdb_mod_it mod_it;
+  syms_bool was_mod_read = syms_false;
+  if (pdb_mod_it_init(&mod_it, pdb)) {
+    if (pdb_mod_it_seek(&mod_it, imod)) {
+      was_mod_read = pdb_mod_it_next(&mod_it, mod);
+    }
+  }
+  return was_mod_read;
+}
+
+PDB_API pdb_file_info
+pdb_file_info_bake_null(void)
+{
+  pdb_file_info result;
+  result.path = pdb_pointer_bake_null();
+  result.chksum_type = PDB_CV_CHECKSUM_NULL;
+  result.chksum = pdb_pointer_bake_null();
+  return result;
+}
+
+PDB_API syms_bool
+pdb_mod_infer_fileid(pdb_mod *mod, U32 fileid, pdb_file_info *fi_out)
+{
+  syms_bool is_resolved = syms_false;
+
+  if (~mod->flags & PDB_MOD_FILECHKSUM_CACHED) {
+    pdb_debug_sec_it dsec_it;
+    if (pdb_debug_sec_it_init(&dsec_it, mod)) {
+      pdb_debug_sec dsec;
+      while (pdb_debug_sec_it_next(&dsec_it, &dsec)) {
+        if (dsec.type == PDB_CV_SS_TYPE_FILE_CHKSUM) {
+          mod->filechksum = dsec;
+          mod->flags |= PDB_MOD_FILECHKSUM_CACHED;
+          break;
+        }
+      }
+    }
+  }
+
+  if (mod->flags & PDB_MOD_FILECHKSUM_CACHED) {
+    pdb_debug_sec *filechksum = &mod->filechksum;
+    if (pdb_stream_seek(&filechksum->stream, fileid)) {
+      pdb_context *pdb = mod->pdb;
+      pdb_cv_file_checksum chksum;
+      if (pdb_stream_read(&filechksum->stream, &chksum, sizeof(chksum))) {
+        pdb_stream strtable = pdb_get_strtable(pdb);
+        if (pdb_stream_seek(&strtable, chksum.name_off)) {
+          fi_out->path = pdb_pointer_bake_stream_str(&strtable);
+          fi_out->chksum_type = chksum.type;
+          fi_out->chksum = pdb_pointer_bake(&filechksum->stream, chksum.len);
+          is_resolved = syms_true;
+        }
+      }
+    }
+  }
+
+  return is_resolved;
+}
+
+PDB_API syms_bool
+pdb_debug_sec_it_init(pdb_debug_sec_it *it, pdb_mod *mod)
+{
+  syms_bool is_inited = syms_false;
+  if (pdb_mod_get_debug_sec(mod, PDB_MOD_SEC_LINES_C13, &it->stream)) {
+    is_inited = syms_true;
+  }
+  return is_inited;
+}
+
+PDB_API syms_bool
+pdb_debug_sec_it_next(pdb_debug_sec_it *it, pdb_debug_sec *sec_out)
+{
+  syms_bool is_sec_valid = syms_false;
+
+  if (pdb_stream_read_u32(&it->stream, &sec_out->type)) {
+    U32 sec_size;
+    if (pdb_stream_read_u32(&it->stream, &sec_size)) {
+      sec_out->stream = pdb_stream_subset(&it->stream, it->stream.off, sec_size);
+      if (sec_out->stream.size > 0) {
+        is_sec_valid = pdb_stream_skip(&it->stream, sec_out->stream.size);
+      }
+    }
+  }
+
+  return is_sec_valid;
+}
+
+PDB_API syms_bool
+pdb_dss_it_init(pdb_dss_it *it, pdb_debug_sec *sec)
+{
+  syms_bool is_inited = syms_false;
+
+  it->type = sec->type;
+  it->stream = sec->stream;
+  it->ex_mode = syms_false;
+
+  if (it->type == PDB_CV_SS_TYPE_INLINE_LINES) {
+    U32 sig;
+    if (pdb_stream_read_u32(&it->stream, &sig)) {
+      it->ex_mode = (sig == PDB_CV_INLINEE_SOURCE_LINE_SIGNATURE_EX);
+      is_inited = syms_true;
+    }
+  } else {
+    is_inited = syms_true;
+  }
+
+  return is_inited;
+}
+
+PDB_API syms_bool
+pdb_dss_it_next_inline(pdb_dss_it *it, pdb_ss_inline *inline_out)
+{
+  if (it->type == PDB_CV_SS_TYPE_INLINE_LINES) {
+    if (it->ex_mode) {
+      pdb_cv_inlinee_srcline_ex srcline;
+      if (pdb_stream_read_struct(&it->stream, &srcline, pdb_cv_inlinee_srcline_ex)) {
+        U32 extra_size = srcline.extra_file_id_count*sizeof(U32);
+        inline_out->inlinee = srcline.inlinee;
+        inline_out->src_ln = srcline.src_ln;
+        inline_out->file_id = srcline.file_id;
+        inline_out->extra_files_count = srcline.extra_file_id_count;
+        inline_out->extra_files = pdb_pointer_bake(&it->stream, extra_size);
+        return pdb_stream_skip(&it->stream, extra_size);
+      }
+    } else {
+      pdb_cv_inlinee_srcline srcline;
+      if (pdb_stream_read_struct(&it->stream, &srcline, pdb_cv_inlinee_srcline)) {
+        inline_out->inlinee = srcline.inlinee;
+        inline_out->src_ln = srcline.src_ln;
+        inline_out->file_id = srcline.file_id;
+        inline_out->extra_files_count = 0;
+        inline_out->extra_files = pdb_pointer_bake_null();
+        return syms_true;
+      }
+    }
+  }
+  return syms_false;
+}
+
+PDB_API pdb_cvdata_token
+pdb_cvdata_token_bake(pdb_sn sn, U32 offset)
+{
+  pdb_cvdata_token cvdata;
+  cvdata.sn = sn;
+  cvdata.soffset = offset;
+  return cvdata;
+}
+
+PDB_API pdb_cvdata_token
+pdb_sym_it_get_token(pdb_sym_it *sym_it)
+{
+  return pdb_cvdata_token_bake(sym_it->stream.sn, pdb_stream_get_abs_off(&sym_it->stream));
+}
+
+PDB_API syms_bool
+pdb_sym_it_init(pdb_sym_it *sym_it, pdb_mod *mod)
+{
+  syms_bool result = syms_false;
+  sym_it->inited_from_token = syms_false;
+  if (pdb_mod_get_debug_sec(mod, PDB_MOD_SEC_SYMS, &sym_it->stream)) {
+    result = syms_true;
+  }
+  return result;
+}
+
+PDB_API syms_bool
+pdb_sym_it_init_token(pdb_sym_it *sym_it, pdb_context *pdb, pdb_cvdata_token token)
+{
+  sym_it->inited_from_token = syms_true;
+
+  if (pdb_stream_init(pdb, token.sn, &sym_it->stream)) {
+    if (pdb_stream_skip(&sym_it->stream, token.soffset)) {
+      return syms_true;
+    }
+  }
+
+  return syms_false;
+}
+
+PDB_API syms_bool
+pdb_sym_it_next(pdb_sym_it *it, pdb_cv_sym_type *type_out, pdb_stream *data_out)
+{
+  syms_bool result = syms_false;
+  pdb_symrec header;
+
+  if (pdb_stream_read_symrec(&it->stream, &header)) {
+    if (type_out) {
+      *type_out = header.type;
+    }
+    if (data_out) {
+      U32 data_size = header.end - it->stream.off;
+      *data_out = pdb_stream_subset(&it->stream, it->stream.off, data_size);
+    }
+
+    if (pdb_stream_seek(&it->stream, header.end)) {
+      result = pdb_stream_align(&it->stream, 4);
+    }
+  }
+
+  return result;
+}
+
+PDB_API syms_bool
+pdb_sym_it_peek(pdb_sym_it *it, pdb_cv_sym_type *type_out, pdb_stream *stream_out)
+{
+  pdb_stream temp = it->stream;
+  syms_bool result = pdb_sym_it_next(it, type_out, stream_out);
+  it->stream = temp;
+  return result;
+}
+
+PDB_API syms_bool
+pdb_sym_it_next_proc(pdb_sym_it *sym_it, pdb_proc *proc_out)
+{
+  pdb_cv_sym_type sym_type;
+  pdb_stream sym_stream;
+  syms_bool is_result_valid = syms_false;
+
+  if (pdb_sym_it_peek(sym_it, &sym_type, &sym_stream)) {
+    switch (sym_type) {
+    case PDB_CV_SYM_GPROC32:
+    case PDB_CV_SYM_LPROC32: {
+      pdb_cv_proc cv_proc;
+      if (pdb_stream_read(&sym_stream, &cv_proc, sizeof(cv_proc))) {
+        proc_out->sec     = cv_proc.sec;
+        proc_out->sec_off = cv_proc.off;
+        proc_out->size    = cv_proc.len;
+        proc_out->itype   = cv_proc.itype;
+        proc_out->flags   = cv_proc.flags;
+        proc_out->name    = pdb_pointer_bake_stream_str(&sym_stream);
+        proc_out->cvdata  = pdb_sym_it_get_token(sym_it);
+        is_result_valid   = syms_true;
+        pdb_sym_it_next(sym_it, 0, 0);
+      }
+    } break;
+
+    case PDB_CV_SYM_GPROC16:
+    case PDB_CV_SYM_LPROC16: {
+      SYMS_ASSERT_NO_SUPPORT;
+    } break;
+
+    case PDB_CV_SYM_GPROC32_16t:
+    case PDB_CV_SYM_LPROC32_16t: {
+      SYMS_ASSERT_NO_SUPPORT;
+    } break;
+
+    default: break;
+    }
+  }
+
+  return is_result_valid;
+}
+
+PDB_API pdb_var
+pdb_var_init(pdb_ti itype, pdb_cv_localsym_flags flags, pdb_encoded_location encoded_va, pdb_string_ref name)
+{
+  pdb_var result;
+
+  result.itype = itype;
+  result.flags = flags;
+  result.encoded_va = encoded_va;
+  result.name = name;
+  result.gaps = pdb_pointer_bake_null();
+
+  return result;
+}
+
+PDB_API pdb_var
+pdb_var_init_null(void)
+{
+  return pdb_var_init(
+      PDB_INVALID_ITYPE,
+      0,
+      pdb_encode_location_for_null(),
+      pdb_pointer_bake_null());
+}
+
+PDB_API syms_bool
+pdb_local_it_init_(pdb_local_it *local_it, pdb_context *pdb, pdb_isec sec, pdb_isec_umm sec_off, pdb_cvdata_token cvdata)
+{
+  pdb_proc proc;
+  pdb_type proc_type;
+  syms_bool is_inited = syms_false;
+
+  pdb_sym_it sym_it;
+  if (pdb_sym_it_init_token(&sym_it, pdb, cvdata)) {
+    if (pdb_sym_it_next_proc(&sym_it, &proc)) {
+      if (pdb_infer_itype(pdb, proc.itype, &proc_type)) {
+        if (proc_type.kind == PDB_TYPE_PROC || proc_type.kind == PDB_TYPE_METHOD) {
+          local_it->pdb           = pdb;
+          local_it->defrange_mode = (proc.flags & PDB_CV_PROC32_FLAG_OPTDBGINFO) != 0;
+          local_it->sym_it        = sym_it;
+          local_it->range_off     = proc.sec_off;
+          local_it->range_len     = proc.size;
+          /* NOTE(nick): Offsets in the symbols are section relative, so to
+           * make our lives easier we just convert PC to relative offset too. */
+          local_it->sec              = sec;
+          local_it->sec_off          = sec_off;
+          local_it->inlinesite_count = 0;
+          local_it->regrel32_count   = 0;
+          local_it->block32_count    = 0;
+
+          if (proc_type.kind == PDB_TYPE_PROC) {
+            local_it->arg_count = proc_type.u.proc.arg_count;
+          } else if (proc_type.kind == PDB_TYPE_METHOD) {
+            local_it->arg_count = proc_type.u.method.arg_count;
+          }
+
+          is_inited = syms_true;
+        }
+      }
+    }
+  }
+
+  return is_inited;
+}
+
+PDB_API syms_bool
+pdb_local_it_init(pdb_local_it *local_it, pdb_context *pdb, pdb_isec sec, pdb_isec_umm sec_off)
+{
+  pdb_proc_it proc_it;
+  syms_bool is_inited = syms_false;
+  pdb_mod_it mod_it;
+
+  if (!local_it || !pdb)
+    return syms_false;
+
+  if (pdb_mod_it_init(&mod_it, pdb)) {
+    pdb_mod mod;
+    while (pdb_mod_it_next(&mod_it, &mod)) {
+      if (pdb_proc_it_init(&proc_it, &mod)) {
+        pdb_proc proc;
+        while (pdb_proc_it_next(&proc_it, &proc)) {
+          if (sec == proc.sec && (sec_off >= proc.sec_off && sec_off < proc.sec_off + proc.size)) {
+            is_inited = pdb_local_it_init_(local_it, pdb, sec, sec_off, proc.cvdata);
+            break;
+          }
+        }
+      }
+    }
+  }
+
+  return is_inited;
+}
+
+PDB_API syms_bool
+pdb_local_it_next(pdb_local_it *local_it, pdb_local_export *export_out)
+{
+  pdb_sym_it *sym_it = &local_it->sym_it;
+  syms_bool is_result_valid = syms_false;
+  syms_bool keep_looping;
+
+  if (local_it->block32_count < 0) {
+    return syms_false;
+  }
+
+  do {
+    pdb_cv_sym_type cv_symbol_type;
+    pdb_stream cv_stream;
+
+    if (!pdb_sym_it_next(sym_it, &cv_symbol_type, &cv_stream))
+      break;
+
+    keep_looping = syms_false;
+
+    switch (cv_symbol_type) {
+    case PDB_CV_SYM_LOCAL: { /* Optimized-mode variable */
+      pdb_cv_localsym localsym;
+
+      if (pdb_stream_read_struct(&cv_stream, &localsym, pdb_cv_localsym)) {
+        pdb_encoded_location location = pdb_encode_location_for_null();
+
+        pdb_cv_sym_type sym_type;
+        pdb_stream defrange_stream;
+
+        while (pdb_sym_it_peek(sym_it, &sym_type, &defrange_stream)) {
+          syms_bool is_enreged = syms_false;
+          pdb_uint regindex = PDB_UINT_MAX;
+          pdb_uint regoff = PDB_UINT_MAX;
+          pdb_cv_lvar_addr_range_t range;
+
+          range.sec = 0;
+          range.off = 0;
+          range.len = 0;
+
+          switch (sym_type) {
+          case PDB_CV_SYM_DEFRANGE_REGISTER_REL: {
+            /* Register-relative variable */
+            pdb_cv_defrange_register_rel rel;
+
+            if (pdb_stream_read_struct(&defrange_stream, &rel, pdb_cv_defrange_register_rel)) {
+              is_enreged = syms_false;
+              regindex   = rel.reg;
+              regoff     = rel.reg_off;
+              range      = rel.range;
+            }
+          } break;
+          case PDB_CV_SYM_DEFRANGE_FRAMEPOINTER_REL: {
+            /* Frame-pointer relative variable */
+            pdb_cv_defrange_frameptr_rel rel;
+
+            if (pdb_stream_read_struct(&defrange_stream, &rel, pdb_cv_defrange_frameptr_rel)) {
+              is_enreged = syms_false;
+              regindex   = PDB_CV_X64_RSP;
+              regoff     = rel.off;
+              range      = rel.range;
+            }
+          } break;
+          case PDB_CV_SYM_DEFRANGE_REGISTER: {
+            /* Symbol for an optimized variable that is stored in a register when PC inside range. */
+            pdb_cv_defrange_reg rel;
+
+            if (pdb_stream_read_struct(&defrange_stream, &rel, pdb_cv_defrange_reg)) {
+              is_enreged = syms_true;
+              regindex   = rel.reg;
+              regoff     = 0;
+              range      = rel.range;
+            }
+          } break;
+          case PDB_CV_SYM_DEFRANGE_SUBFIELD_REGISTER: {
+            pdb_cv_defrange_subfield_reg rel;
+
+            if (pdb_stream_read_struct(&defrange_stream, &rel, pdb_cv_defrange_subfield_reg)) {
+              is_enreged  = syms_false;
+              regindex    = rel.reg;
+              regoff      = PDB_UINT_MAX;
+              range       = rel.range;
+            }
+          } break;
+          case PDB_CV_SYM_DEFRANGE_FRAMEPOINTER_REL_FULL_SCOPE: {
+            /* Frame-pointer relative variable, but without bounds on range */
+            pdb_cv_defrange_frameptr_rel_full_scope rel;
+
+            if (pdb_stream_read_struct(&defrange_stream, &rel, pdb_cv_defrange_frameptr_rel_full_scope)) {
+              is_enreged  = syms_false;
+              regindex    = PDB_CV_X64_RSP;
+              regoff      = rel.off;
+              syms_memset(&range, 0xff, sizeof(range));
+            }
+          } break;
+          default: {
+            defrange_stream.flags |= PDB_STREAM_FLAGS_READ_FAILED;
+          } break;
+          }
+
+          /* NOTE(nick): First of all make sure that read was valid */
+          if ((defrange_stream.flags & PDB_STREAM_FLAGS_READ_FAILED) == PDB_STREAM_FLAGS_READ_FAILED)
+            break;
+
+          if (!pdb_sym_it_next(sym_it, NULL, NULL)) {
+            #if defined(SYMS_PARANOID)
+            /* NOTE(nick): This symbol was read with pdb_sym_it_peek, 
+             * no crash in this case can occur */
+            SYMS_INVALID_CODE_PATH;
+            #endif
+          }
+
+          if (is_enreged)
+            location = pdb_encode_location_for_enreged(local_it->pdb, regindex);
+          else
+            location = pdb_encode_location_for_regrel(local_it->pdb, regindex, regoff);
+
+          if (~location.flags & PDB_ENCODED_LOCATION_FLAG_OUTSIDE_RANGE)
+            break;
+        }
+
+        export_out->type = PDB_LOCAL_EXPORT_VAR;
+        export_out->u.var = pdb_var_init(
+                                      localsym.itype, 
+                                      localsym.flags,
+                                      location,
+                                      pdb_pointer_bake_stream_str(&cv_stream));
+
+        {
+          pdb_uint gap_count = (defrange_stream.size - defrange_stream.off) / sizeof(pdb_cv_lvar_addr_gap_t);
+          if (gap_count > 0) {
+            export_out->u.var.gaps = pdb_pointer_bake(&defrange_stream, gap_count * sizeof(pdb_cv_lvar_addr_gap_t));
+          } else {
+            export_out->u.var.gaps = pdb_pointer_bake_null();
+          }
+        }
+        is_result_valid = syms_true;
+      }
+    } break;
+
+    case PDB_CV_SYM_LDATA32: { /* static variable */
+      pdb_cv_datasym32 cv_data;
+
+      if (pdb_stream_read_struct(&cv_stream, &cv_data, pdb_cv_datasym32)) {
+        export_out->type = PDB_LOCAL_EXPORT_VAR;
+        export_out->u.var = pdb_var_init(
+                                    cv_data.itype,
+                                    PDB_CV_LOCALSYM_FLAG_STATIC,
+                                    pdb_encode_location_for_secoff(local_it->pdb, cv_data.sec, cv_data.sec_off),
+                                    pdb_pointer_bake_stream_str(&cv_stream));
+        is_result_valid = syms_true;
+      }
+    } break;
+
+    case PDB_CV_SYM_REGREL32: { /* Debug-mode variable */
+      if (!local_it->defrange_mode) {
+        pdb_cv_regrel32 cv_regrel;
+
+        if (pdb_stream_read_struct(&cv_stream, &cv_regrel, pdb_cv_regrel32)) {
+          export_out->type = PDB_LOCAL_EXPORT_VAR;
+          export_out->u.var = pdb_var_init(
+              cv_regrel.itype,
+              /* no flags */ 0,
+              pdb_encode_location_for_regrel(local_it->pdb, cv_regrel.reg, cv_regrel.reg_off),
+              pdb_pointer_bake_stream_str(&cv_stream));
+          is_result_valid = syms_true;
+        } 
+
+        ++local_it->regrel32_count;
+      } 
+    } break;
+
+    case PDB_CV_SYM_BLOCK32: { /* Scope start*/
+      pdb_cv_blocksym32 cv_block;
+
+      if (pdb_stream_read_struct(&cv_stream, &cv_block, pdb_cv_blocksym32)) {
+        SymsAddr block_base = 0;
+
+        pdb_build_va(local_it->pdb, cv_block.sec, cv_block.off, &block_base);
+        export_out->type = PDB_LOCAL_EXPORT_SCOPE;
+        export_out->u.scope.inst_lo = block_base;
+        export_out->u.scope.inst_hi = block_base + cv_block.len;
+
+		    local_it->block32_count += 1;
+        is_result_valid = syms_true;
+      }            
+    } break;
+
+    case PDB_CV_SYM_END: { /* Scope end */
+      local_it->block32_count -= 1;
+      if (local_it->block32_count >= 0) {
+        export_out->type = PDB_LOCAL_EXPORT_SCOPE_END;
+        is_result_valid = syms_true;
+      }
+    } break;
+
+#if 0
+    case PDB_CV_SYM_INLINESITE: {
+      local_it->block32_count += 1;
+    } break;
+
+    case PDB_CV_SYM_INLINESITE_END: {
+      local_it->block32_count -= 1;
+    } break;
+#endif
+
+#if 0
+    case PDB_CV_SYM_FRAMEPROC:   break;
+    case PDB_CV_SYM_FRAMECOOKIE: break;
+#endif
+    }
+
+    if (PDB_STREAM_READ_OR_SEEK_FAILED(cv_stream.flags))
+      break;
+
+    if (!is_result_valid) {
+      keep_looping = local_it->block32_count >= 0;
+    }
+  } while (keep_looping);
+
+  return is_result_valid;
+}
+
+PDB_API syms_bool
+pdb_global_it_init(pdb_global_it *global_it, pdb_context *pdb)
+{
+  syms_bool is_inited = syms_false;
+
+  if (pdb->globals_array && pdb->globals_array_num > 0) {
+    global_it->pdb       = pdb;
+    global_it->hr_index  = 0;
+    global_it->hr        = NULL;
+
+    if (pdb_stream_init(pdb, pdb->dbi.symrec_sn, &global_it->stream)) {
+      is_inited = syms_true;
+    }
+  }
+
+  return is_inited;
+}
+
+PDB_API syms_bool
+pdb_global_it_next(pdb_global_it *global_it, pdb_var *var_out)
+{
+  pdb_stream *hr_stream = &global_it->stream;
+  pdb_context *pdb = hr_stream->pdb;
+
+  while (global_it->hr_index < pdb->globals_array_num) {
+    if (!global_it->hr) {
+      SYMS_ASSERT(pdb->globals_array);
+      global_it->hr = pdb->globals_array[global_it->hr_index++];
+    }
+
+    for (; global_it->hr != NULL; global_it->hr = global_it->hr->next) {
+      pdb_symrec sym;
+
+      if (!pdb_stream_seek(hr_stream, global_it->hr->off)) {
+        SYMS_ASSERT_FAILURE("invalid hr->off");
+        continue;
+      }
+      if (!pdb_stream_read_symrec(hr_stream, &sym)) {
+        SYMS_ASSERT_FAILURE("reading symbol record failed");
+        continue;
+      }
+
+      switch (sym.type) {
+      case PDB_CV_SYM_LDATA32:
+      case PDB_CV_SYM_GDATA32: {
+        pdb_cv_datasym32 datasym;
+        pdb_pointer data_name;
+        pdb_encoded_location encoded_loc;
+
+        if (!pdb_stream_read_struct(hr_stream, &datasym, pdb_cv_datasym32)) {
+          SYMS_ASSERT_FAILURE("invalid stream data, cannot read cv_datasym32");
+          break;
+        }
+
+        data_name = pdb_pointer_bake_stream_str(hr_stream);
+        encoded_loc = pdb_encode_location_for_datasym32(pdb, &datasym);
+        *var_out = pdb_var_init(datasym.itype, 0, encoded_loc, data_name);
+
+        /* NOTE(nick): This list might contain older versions of a variable. 
+         * Newest always comes first, hence if we find a duplicate name 
+         * after current variable it's an older version and it can be ignored. */
+        for (; global_it->hr != NULL; global_it->hr = global_it->hr->next) {
+          pdb_symrec sym_next;
+          pdb_pointer next_name;
+
+          if (!pdb_stream_seek(hr_stream, global_it->hr->off)) {
+            break;
+          }
+          if (!pdb_stream_read_symrec(hr_stream, &sym_next)) {
+            break;
+          }
+          if (sym_next.type != sym.type) {
+            break;
+          }
+          if (!pdb_stream_skip(hr_stream, sizeof(pdb_cv_datasym32))) {
+            break;
+          }
+          next_name = pdb_pointer_bake_stream_str(hr_stream);
+          if (!pdb_pointer_strcmp(pdb, &var_out->name, &next_name)) {
+            break;
+          }
+        }
+
+        return syms_true;
+      }
+
+      case PDB_CV_SYM_GDATA32_16t:
+      case PDB_CV_SYM_LDATA32_16t: {
+        SYMS_ASSERT_FAILURE("encountered an unsupported obsolete symbol");
+      } break;
+
+      case PDB_CV_SYM_LDATA16:
+      case PDB_CV_SYM_GDATA16: {
+        SYMS_ASSERT_FAILURE("encountered an unsupported obsolete symbol");
+      } break;
+
+      default: {
+        /* NOTE(nick): This stream is mishmash of symbols, ignore this symbol */
+      } break;
+      }
+    }
+  }
+
+  return syms_false;
+}
+
+PDB_API syms_bool
+pdb_global_from_name(pdb_context *pdb, const char *name, U32 name_size, pdb_var *var_out)
+{
+  SymsString name_str = syms_string_init(name, name_size);
+  pdb_global_it it;
+
+  if (pdb_global_it_init(&it, pdb)) {
+    U32 name_index = pdb_hashV1_bytes(name, name_size, pdb->globals_array_num); 
+
+    it.hr_index = name_index;
+    it.hr = NULL;
+
+    while (pdb_global_it_next(&it, var_out)) {
+      if (it.hr_index != (name_index + 1))
+        break;
+
+      if (pdb_strcmp_pointer(pdb, name_str, &var_out->name))
+        return syms_true;
+    }
+  }
+
+  return syms_false;
+}
+
+PDB_API syms_bool
+pdb_const_from_name(pdb_context *pdb, const char *name, U32 name_size, pdb_const_value *const_out)
+{
+  SymsString name_str = syms_string_init(name, name_size);
+  pdb_const_it it;
+
+  if (pdb_const_it_init(&it, pdb)) {
+    U32 name_index = pdb_hashV1_bytes(name, name_size, pdb->globals_array_num);
+
+    it.index = name_index;
+    it.hr = NULL;
+
+    while (pdb_const_it_next(&it, const_out)) {
+      if (pdb_strcmp_pointer(pdb, name_str, &const_out->name)) {
+        return syms_true;
+      }
+      if (name_index != it.index) {
+        break;
+      }
+    }
+  }
+
+  return syms_false;
+}
+
+PDB_API syms_bool
+pdb_file_it_init(pdb_file_it *file_it, pdb_context *pdb)
+{
+  file_it->pdb        = pdb;
+  file_it->strtable   = pdb_get_strtable(pdb);
+  file_it->stroffs    = pdb_get_stroffs(pdb);
+  file_it->off_count  = file_it->stroffs.size/sizeof(pdb_stroff);
+  file_it->num_read   = 0;
+  return syms_true;
+}
+
+PDB_API syms_bool
+pdb_file_it_next(pdb_file_it *file_it, pdb_pointer *file_out)
+{
+  for (;file_it->num_read < file_it->off_count; file_it->num_read += 1) {
+    pdb_stroff stroff;
+    U8 b0, b1, b2;
+    syms_bool is_path;
+
+    if (!pdb_stream_seek(&file_it->stroffs, file_it->num_read*sizeof(pdb_stroff))) break;
+    if (!pdb_stream_read(&file_it->stroffs, &stroff, sizeof(stroff))) break;
+    if (!pdb_stream_seek(&file_it->strtable, stroff)) break;
+
+    *file_out = pdb_pointer_bake_stream_str(&file_it->strtable);
+
+    b0 = pdb_pointer_read_u08(file_it->pdb, file_out, 0);
+    b1 = pdb_pointer_read_u08(file_it->pdb, file_out, 1);
+    b2 = pdb_pointer_read_u08(file_it->pdb, file_out, 2);
+
+    /* Check if this is a windows path. */
+    is_path = (syms_is_alpha_ascii((char)b0) && b1 == ':' && (b2 == '\\' || b2 == '/'));
+
+    /* Check if this a unix path. */
+    is_path = (is_path || (b0 == '/' && b1 != 0));
+    if (!is_path) continue;
+    
+    file_it->num_read += 1;
+    return syms_true;
+  }
+
+  return syms_false;
+}
+
+PDB_API syms_bool
+pdb_proc_from_stream(pdb_cv_sym_type cvtype, pdb_stream *cvdata, pdb_proc *proc_out)
+{
+  syms_bool is_result_valid = syms_false;
+
+  switch (cvtype) {
+  case PDB_CV_SYM_GPROC32:
+  case PDB_CV_SYM_LPROC32: {
+    pdb_cvdata_token cvdata_token = pdb_cvdata_token_bake(cvdata->sn, pdb_stream_get_abs_off(cvdata) - 4);
+    pdb_cv_proc cvproc;
+
+    if (pdb_stream_read(cvdata, &cvproc, sizeof(cvproc))) {
+      proc_out->sec       = cvproc.sec;
+      proc_out->sec_off   = cvproc.off;
+      proc_out->size      = cvproc.len;
+      proc_out->itype     = cvproc.itype;
+      proc_out->flags     = cvproc.flags;
+      proc_out->name      = pdb_pointer_bake_stream_str(cvdata);
+      proc_out->cvdata    = cvdata_token;
+      is_result_valid = syms_true;
+    }
+  } break;
+
+  case PDB_CV_SYM_GPROC16:
+  case PDB_CV_SYM_LPROC16: {
+    SYMS_ASSERT_FAILURE("no support for 16-bit procedure symbol");
+  } break;
+
+  case PDB_CV_SYM_GPROC32_16t:
+  case PDB_CV_SYM_LPROC32_16t: {
+    SYMS_ASSERT_FAILURE("no support for 16T procedure symbol");
+  } break;
+
+  default: break;
+  }
+
+  return is_result_valid;
+}
+
+PDB_API syms_bool
+pdb_proc_it_init(pdb_proc_it *proc_it, pdb_mod *mod)
+{
+  return pdb_sym_it_init(&proc_it->sym_it, mod);
+}
+PDB_API syms_bool
+pdb_proc_it_next(pdb_proc_it *proc_it, pdb_proc *proc_out)
+{
+  syms_bool is_proc_found = syms_false;
+
+  for (;;) {
+    pdb_cv_sym_type cvtype;
+    pdb_stream cvdata;
+
+    if (!pdb_sym_it_next(&proc_it->sym_it, &cvtype, &cvdata)) {
+      break;
+    }
+
+    if (pdb_proc_from_stream(cvtype, &cvdata, proc_out)) {
+      is_proc_found = syms_true;
+      break;
+    }
+  }
+
+  return is_proc_found;
+}
+
+PDB_API syms_bool
+pdb_proc_from_cvdata(pdb_context *pdb, pdb_cvdata_token cvdata, pdb_proc *proc_out)
+{
+  pdb_sym_it sym_it;
+  if (pdb_sym_it_init_token(&sym_it, pdb, cvdata)) {
+    return pdb_sym_it_next_proc(&sym_it, proc_out);
+  }
+  return syms_false;
+}
+
+PDB_API syms_bool
+pdb_proc_from_name_(pdb_context *pdb, pdb_pointer *name, pdb_proc *proc_out)
+{
+  pdb_stream stream;
+
+  if (pdb_stream_init(pdb, pdb->dbi.symrec_sn, &stream)) {
+    pdb_gsi_hr *hr;
+    U32 index;
+
+    if (pdb->globals_array) {
+      index = pdb_hashV1_pointer(pdb, name, pdb->globals_array_num);
+      for (hr = pdb->globals_array[index]; hr != 0; hr = hr->next) {
+        pdb_symrec sym;
+        pdb_cv_symref2 ref;
+        pdb_mod mod;
+        pdb_cvdata_token cvdata;
+
+        if (!pdb_stream_seek(&stream, hr->off)) {
+          continue;
+        }
+        if (!pdb_stream_read_symrec(&stream, &sym)) {
+          continue;
+        }
+        if (sym.type != PDB_CV_SYM_PROCREF && sym.type != PDB_CV_SYM_LPROCREF) {
+          continue;
+        }
+        if (!pdb_stream_read_struct(&stream, &ref, pdb_cv_symref2)) {
+          continue;
+        }
+        if (!pdb_stream_strcmp_pointer(&stream, name)) {
+          continue;
+        }
+        SYMS_ASSERT(ref.imod > 0);
+        if (!pdb_mod_init(&mod, pdb, (pdb_imod)(ref.imod - 1))) {
+          continue;
+        }
+        cvdata = pdb_cvdata_token_bake(mod.sn, ref.sym_off);
+        if (pdb_proc_from_cvdata(pdb, cvdata, proc_out)) {
+          return syms_true;
+        }
+      }
+    }
+
+    if (pdb->publics_array) {
+#if 1
+      index = pdb_hashV1_pointer(pdb, name, pdb->publics_array_num);
+      for (hr = pdb->publics_array[index]; hr != 0; hr = hr->next) {
+        pdb_symrec sym;
+        pdb_cv_symref2 ref;
+        pdb_mod mod;
+        pdb_cvdata_token cvdata;
+
+        if (!pdb_stream_seek(&stream, hr->off)) {
+          continue;
+        }
+        if (!pdb_stream_read_symrec(&stream, &sym)) {
+          continue;
+        }
+        if (sym.type != PDB_CV_SYM_PROCREF && sym.type != PDB_CV_SYM_LPROCREF) {
+          continue;
+        }
+        if (!pdb_stream_read_struct(&stream, &ref, pdb_cv_symref2)) {
+          continue;
+        }
+        if (!pdb_stream_strcmp_pointer(&stream, name)) {
+          continue;
+        }
+        SYMS_ASSERT(ref.imod > 0);
+        if (!pdb_mod_init(&mod, pdb, (pdb_imod)(ref.imod - 1))) {
+          continue;
+        }
+        cvdata = pdb_cvdata_token_bake(mod.sn, ref.sym_off);
+        if (pdb_proc_from_cvdata(pdb, cvdata, proc_out)) {
+          return syms_true;
+        }
+      }
+#else
+      U32 counter = 0;
+      for (index = 0; index < pdb->publics_array_num; ++index) {
+        for (hr = pdb->publics_array[index]; hr != 0; hr = hr->next) {
+          pdb_symrec sym;
+          pdb_cv_symref2 ref;
+          pdb_mod mod;
+          pdb_cvdata_token cvdata;
+
+          counter += 1;
+          if (counter == 1161) {
+            int x = 0;
+          }
+
+          if (!pdb_stream_seek(&stream, hr->off)) {
+            continue;
+          }
+          if (!pdb_stream_read_symrec(&stream, &sym)) {
+            continue;
+          }
+          if (sym.type != PDB_CV_SYM_PROCREF && sym.type != PDB_CV_SYM_LPROCREF) {
+            continue;
+          }
+          if (!pdb_stream_read_struct(&stream, &ref, pdb_cv_symref2)) {
+            continue;
+          }
+          if (!pdb_stream_strcmp_pointer(&stream, name)) {
+            continue;
+          }
+          SYMS_ASSERT(ref.imod > 0);
+          if (!pdb_mod_init(&mod, pdb, (pdb_imod)(ref.imod - 1))) {
+            continue;
+          }
+          cvdata = pdb_cvdata_token_bake(mod.sn, ref.sym_off);
+          if (pdb_proc_from_cvdata(pdb, cvdata, proc_out)) {
+            return syms_true;
+          }
+        }
+      }
+#endif
+    }
+  }
+
+  return syms_false;
+}
+
+PDB_API syms_bool
+pdb_proc_from_va(pdb_context *pdb, SymsAddr va, pdb_proc *proc_out)
+{
+#if 0
+  pdb_pointer name;
+
+  if (pdb_find_nearest_sym(pdb, va, &name)) {
+    if (pdb_proc_from_name_(pdb, &name, proc_out)) {
+      return syms_true;
+    }
+  }
+
+  return syms_false;
+#else
+  pdb_proc_it proc_it;
+  pdb_isec sec;
+  pdb_isec_umm sec_off;
+  syms_bool result = syms_false;
+
+  if (pdb_build_sec_off(pdb, va, &sec, &sec_off)) {
+    pdb_mod_it mod_it;
+    if (pdb_mod_it_init(&mod_it, pdb)) {
+      pdb_mod mod;
+      while (pdb_mod_it_next(&mod_it, &mod)) {
+        if (pdb_proc_it_init(&proc_it, &mod)) {
+          pdb_proc t;
+
+          while (pdb_proc_it_next(&proc_it, &t)) {
+            if (sec == t.sec) {
+              if (sec_off >= t.sec_off && sec_off < t.sec_off + t.size) {
+                *proc_out = t;
+                result = syms_true;
+                break;
+              }
+            }
+          }
+        }
+      }
+    }
+  }
+  return result;
+#endif
+}
+
+PDB_API syms_bool
+pdb_proc_from_name(pdb_context *pdb, const char *name, pdb_uint name_len, pdb_proc *proc_out)
+{
+  pdb_pointer name_pointer = pdb_pointer_bake_str(syms_string_init(name, name_len));
+  return pdb_proc_from_name_(pdb, &name_pointer, proc_out);
+}
+
+PDB_API syms_bool
+pdb_mod_find_inline_srcpos(pdb_mod *mod, pdb_cv_itemid inlinee, pdb_file_info *fi_out, pdb_uint *ln_out)
+{
+  pdb_debug_sec_it debug_sec_it;
+  pdb_debug_sec debug_sec;
+  if (pdb_debug_sec_it_init(&debug_sec_it, mod)) {
+    while (pdb_debug_sec_it_next(&debug_sec_it, &debug_sec)) {
+      if (debug_sec.type == PDB_CV_SS_TYPE_INLINE_LINES) {
+        pdb_dss_it dss_it;
+        pdb_ss_inline inline_data;
+        if (pdb_dss_it_init(&dss_it, &debug_sec)) {
+          while (pdb_dss_it_next_inline(&dss_it, &inline_data)) {
+            if (inline_data.inlinee == inlinee) {
+              if (pdb_mod_infer_fileid(mod, inline_data.file_id, fi_out)) {
+                *ln_out = inline_data.src_ln;
+                return syms_true;
+              }
+            }
+          }
+        }
+      }
+    }
+  }
+
+  return syms_false;
+}
+
+PDB_API syms_bool
+pdb_inline_it_init(pdb_inline_it *inline_it, pdb_context *pdb, pdb_cvdata_token cvproc, SymsAddr proc_pc)
+{
+  syms_bool is_inited = syms_false;
+  if (pdb_sym_it_init_token(&inline_it->sym_it, pdb, cvproc)) {
+    if (pdb_sym_it_next_proc(&inline_it->sym_it, &inline_it->proc)) {
+      pdb_imod imod;
+      if (pdb_imod_from_isec(pdb, inline_it->proc.sec, inline_it->proc.sec_off, &imod)) {
+        inline_it->proc_pc = proc_pc;
+        is_inited = pdb_mod_init(&inline_it->mod, pdb, imod);
+      }
+    }
+  }
+  return is_inited;
+}
+
+PDB_API syms_bool
+pdb_inline_it_read(pdb_inline_it *it, pdb_cv_inlinesym *sym_out, pdb_stream *ba_out)
+{
+  if (it->site_count >= 0) {
+    pdb_cv_sym_type sym_type;
+    while (pdb_sym_it_next(&it->sym_it, &sym_type, ba_out)) {
+      switch (sym_type) {
+      case PDB_CV_SYM_INLINESITE: {
+        if (pdb_stream_read(ba_out, sym_out, sizeof(*sym_out))) {
+          *ba_out = pdb_stream_subset(ba_out, sizeof(*sym_out), ba_out->size - sizeof(*sym_out));
+          it->site_count += 1;
+          return (ba_out->size > 0);
+        }
+        SYMS_ASSERT_FAILURE("cannot parse PDB_CV_SYM_INLINESITE");
+      } break;
+
+      case PDB_CV_SYM_INLINESITE2: {
+        pdb_cv_inlinesym2 sym;
+        if (pdb_stream_read(ba_out, &sym, sizeof(sym))) {
+          *ba_out = pdb_stream_subset(ba_out, sizeof(*sym_out), ba_out->size - sizeof(*sym_out));
+          if (ba_out->size > 0) {
+            sym_out->parent_offset = sym.parent_offset;
+            sym_out->end_offset    = sym.end_offset;
+            sym_out->inlinee       = sym.inlinee;
+
+            it->site_count += 1;
+            return syms_true;
+          } 
+        }
+      } break;
+
+      case PDB_CV_SYM_INLINESITE_END: {
+        it->site_count -= 1;
+        if (it->site_count < 0) {
+          return syms_false;
+        }
+      } break;
+
+      case PDB_CV_SYM_END: {
+        return syms_false;
+      } break;
+
+      default: /* ignore */ break;
+      }
+    }
+  }
+  return syms_false;
+}
+
+PDB_API syms_bool
+pdb_inline_it_next(pdb_inline_it *it, pdb_inline_site *site_out)
+{
+  syms_bool sym_consumed = syms_false;
+  pdb_cv_inlinesym inline_sym;
+  pdb_stream binary_annots;
+
+  while (pdb_inline_it_read(it, &inline_sym, &binary_annots)) {
+    syms_bool keep_parsing = syms_true;
+
+    pdb_uint range_ln = 0;
+    pdb_uint code_offset_base = 0;
+    pdb_uint file_id = PDB_UINT_MAX;
+
+    pdb_uint nearest_range_off  = 0;
+    pdb_uint nearest_range_size = 0;
+    pdb_uint nearest_ln = 0;
+    syms_bool site_found = syms_false;
+
+    while (keep_parsing) {
+      /* TODO(nick; Oct 3 2019): do we need to export anything else but range? */
+      pdb_uint op = PDB_CV_BA_OP_MAX;
+      pdb_uint v = 0;
+      pdb_uint range_kind = 0;
+      pdb_int line_offset = 0;
+      pdb_int colm_offset = 0;
+      pdb_uint code_offset = 0;
+      pdb_uint code_length = 0;
+
+      pdb_stream_read_uleb32(&binary_annots, &op);
+      switch (op) {
+      case PDB_CV_BA_OP_CODE_OFFSET: {
+        keep_parsing = pdb_stream_read_uleb32(&binary_annots, &code_offset);
+      } break;
+
+      case PDB_CV_BA_OP_CHANGE_CODE_OFFSET_BASE: {
+        keep_parsing = pdb_stream_read_uleb32(&binary_annots, &code_offset_base);
+      } break;
+
+      case PDB_CV_BA_OP_CHANGE_CODE_OFFSET: {
+        keep_parsing = pdb_stream_read_uleb32(&binary_annots, &code_offset);
+      } break;
+
+      case PDB_CV_BA_OP_CHANGE_CODE_LENGTH: {
+        keep_parsing = pdb_stream_read_uleb32(&binary_annots, &code_length);
+      } break;
+
+      case PDB_CV_BA_OP_CHANGE_FILE: {
+        keep_parsing = pdb_stream_read_uleb32(&binary_annots, &file_id);
+      } break;
+
+      case PDB_CV_BA_OP_CHANGE_LINE_OFFSET: {
+        keep_parsing = pdb_stream_read_sleb32(&binary_annots, &line_offset);
+      } break;
+
+      case PDB_CV_BA_OP_CHANGE_LINE_END_DELTA: {
+        keep_parsing = pdb_stream_read_uleb32(&binary_annots, &v);
+        line_offset = (S32)v;
+      } break;
+
+      case PDB_CV_BA_OP_CHANGE_RANGE_KIND: {
+        /*  0 this is a statement.
+         *  1 his is an expression. */
+        keep_parsing = pdb_stream_read_uleb32(&binary_annots, &range_kind);
+      } break;
+
+      case PDB_CV_BA_OP_CHANGE_COLUMN_START: {
+        keep_parsing = pdb_stream_read_uleb32(&binary_annots, &v);
+        colm_offset = (S32)v;
+      } break;
+
+      case PDB_CV_BA_OP_CHANGE_COLUMN_END_DELTA: {
+        keep_parsing = pdb_stream_read_sleb32(&binary_annots, &colm_offset);
+      } break;
+
+      case PDB_CV_BA_OP_CHANGE_CODE_OFFSET_AND_LINE_OFFSET: {
+        pdb_uint annotation = 0;
+        keep_parsing = pdb_stream_read_uleb32(&binary_annots, &annotation);
+        line_offset = (pdb_int)(annotation >> 4);
+        code_offset = (pdb_uint)(annotation & 0xf);
+      } break;
+
+      case PDB_CV_BA_OP_CHANGE_CODE_LENGTH_AND_CODE_OFFSET: {
+        keep_parsing = syms_false;
+        if (pdb_stream_read_uleb32(&binary_annots, &code_length)) {
+          if (pdb_stream_read_uleb32(&binary_annots, &code_offset)) {
+            keep_parsing = syms_true;
+          }
+        }
+      } break;
+
+      case PDB_CV_BA_OP_CHANGE_COLUMN_END: {
+        keep_parsing = pdb_stream_read_uleb32(&binary_annots, &v);
+        colm_offset = (S32)v;
+      } break;
+
+      case PDB_CV_BA_OP_END: {
+        keep_parsing = syms_false; 
+      } break;
+
+      default: {
+        range_ln = 0;
+        nearest_ln = 0;
+        keep_parsing = syms_false;
+      } break;
+      }
+
+      if (keep_parsing) {
+        pdb_offset range_lo = it->proc.sec_off + code_offset_base + code_offset;
+        pdb_offset range_hi = it->proc.sec_off + code_offset_base + code_length;
+        if (it->proc_pc >= range_lo && it->proc_pc < range_hi) {
+          nearest_ln         = range_ln;
+          nearest_range_off  = range_lo;
+          nearest_range_size = range_hi - range_lo;
+          site_found = syms_true;
+        }
+        range_ln += line_offset;
+      }
+    }
+
+    if (site_found) {
+      pdb_type inlinee_type;
+
+      if (!pdb_infer_itemid(binary_annots.pdb, inline_sym.inlinee, &inlinee_type)) {
+        SYMS_ASSERT_FAILURE_PARANOID("uable to infer itemid for inline site");
+        continue;
+      }
+
+      #if 0
+      pdb_debug_sec_it debug_sec_it;
+      pdb_debug_sec debug_sec;
+      if (pdb_debug_sec_it_init(&debug_sec_it, mod)) {
+        while (pdb_debug_sec_it_next(&debug_sec_it, &debug_sec)) {
+          if (debug_sec.type == PDB_CV_SS_TYPE_INLINE_LINES) {
+            pdb_dss_it dss_it;
+            pdb_ss_inline inline_data;
+            
+            if (pdb_dss_it_init(&dss_it, &debug_sec)) {
+              while (pdb_dss_it_next_inline(&dss_it, &inline_data)) {
+                if (inline_data.inlinee == inline_sym.inlinee) {
+                  if (pdb_mod_infer_fileid(mod, inline_data.file_id, fi_out)) {
+                    *ln_out = inline_data.src_ln;
+                    return syms_true;
+                  }
+                }
+              }
+            }
+          }
+        }
+      }
+      #else
+      if (!pdb_mod_find_inline_srcpos(&it->mod, inline_sym.inlinee, &site_out->fi, &site_out->ln_at_pc)) {
+        SYMS_ASSERT_FAILURE_PARANOID("unable to find inline source position");
+        continue;
+      }
+      #endif
+      
+      if (file_id != PDB_UINT_MAX) {
+        pdb_mod_infer_fileid(&it->mod, file_id, &site_out->fi);
+      }
+
+      if (inlinee_type.kind == PDB_TYPE_FUNCID) {
+        site_out->name  = inlinee_type.u.funcid.name;
+        site_out->itype = inlinee_type.u.funcid.itype;
+      } else if (inlinee_type.kind == PDB_TYPE_MFUNCID) {
+        site_out->name  = inlinee_type.u.mfuncid.name;
+        site_out->itype = inlinee_type.u.mfuncid.itype;  
+      }
+      site_out->sec       = it->proc.sec;
+      site_out->sec_off   = nearest_range_off;
+      site_out->size      = nearest_range_size;
+      site_out->ln_at_pc  = site_out->ln_at_pc + nearest_ln;
+      site_out->cvdata    = pdb_sym_it_get_token(&it->sym_it);
+
+      sym_consumed = syms_true;
+
+      break;
+    }
+  }
+
+  return sym_consumed;
+}
+
+PDB_API syms_bool
+pdb_const_it_init(pdb_const_it *const_it, pdb_context *pdb)
+{
+  syms_bool is_inited = syms_false;
+
+  if (pdb_stream_init(pdb, pdb->dbi.symrec_sn, &const_it->stream)) {
+    const_it->pdb = pdb;
+    const_it->index = 0;
+    const_it->hr = 0;
+    is_inited = syms_true;
+  }
+
+  return is_inited;
+}
+
+PDB_API syms_bool
+pdb_const_it_next(pdb_const_it *const_it, pdb_const_value *const_out)
+{
+  syms_bool is_next_valid = syms_false;
+
+  pdb_context *pdb = const_it->pdb;
+  pdb_stream *stream = &const_it->stream;
+
+  for (;;) {
+    pdb_gsi_hr *hr = const_it->hr;
+
+    pdb_cv_constsym constsym;
+    pdb_numeric value;
+
+    while (!hr) {
+      if (const_it->index >= pdb->globals_array_num) {
+        break;
+      }
+      hr = pdb->globals_array[const_it->index++];
+    }
+
+    if (!hr) {
+      break;
+    }
+    const_it->hr = hr->next;
+
+    if (pdb_stream_seek(stream, hr->off)) {
+      pdb_symrec symrec;
+      if (pdb_stream_read_symrec(stream, &symrec)) {
+        if (symrec.type == PDB_CV_SYM_CONSTANT) {
+          if (pdb_stream_read(stream, &constsym, sizeof(constsym))) {
+            if (pdb_stream_read_numeric(stream, &value)) {
+              const_out->name = pdb_pointer_bake_stream_str(stream);
+              const_out->itype = constsym.itype;
+              const_out->value = value;
+              is_next_valid = syms_true;
+              break;
+            }
+          }
+        }
+      }
+    }
+  }
+
+  return is_next_valid;
+}
+
+PDB_API syms_bool
+pdb_arg_it_init(pdb_arg_it *iter, pdb_context *pdb, pdb_cv_itype args_itype)
+{
+  syms_bool is_inited = syms_false;
+  pdb_type itype_info;
+  if (pdb_infer_itype(pdb, args_itype, &itype_info)) {
+    if (itype_info.kind == PDB_TYPE_ARGLIST) {
+      iter->itypes = itype_info.u.arglist.itypes;
+      iter->pdb    = pdb;
+      iter->idx    = 0;
+      is_inited = syms_true;
+    }
+  }
+  return is_inited;
+}
+
+PDB_API syms_bool
+pdb_arg_it_next(pdb_arg_it *iter, pdb_cv_itype *arg_out)
+{
+  pdb_uint itype_off;
+  pdb_uint itype_size;
+  syms_bool is_read;
+
+  *arg_out = 0;
+  itype_off = iter->idx * sizeof(*arg_out);
+  itype_size = pdb_pointer_read(iter->pdb, &iter->itypes, itype_off, arg_out, sizeof(*arg_out));
+  is_read = (itype_size == sizeof(*arg_out));
+  if (is_read) {
+    if (*arg_out == 0) {
+      /* There is no documentation for this, but it looks like
+       * variadic argument has itype of 0. */
+      *arg_out = PDB_ITYPE_VARIADIC;
+    }
+    iter->idx += 1;
+  }
+
+  return is_read;
+}
+
+PDB_API syms_bool
+pdb_member_it_init(pdb_member_it *member_it, pdb_context *pdb, pdb_type_info_udt *udt)
+{
+  syms_bool is_inited = syms_false;
+
+  if (pdb_stream_init(pdb, PDB_DEFAULT_STREAM_TPI, &member_it->stream)) {
+    pdb_uint off;
+
+    if (pdb_tm_get_itype_offset(pdb, udt->field_itype, &off)) {
+      if (pdb_stream_seek(&member_it->stream, off)) {
+        pdb_symrec symrec;
+
+        if (pdb_stream_read_symrec(&member_it->stream, &symrec)) {
+          member_it->pdb = pdb;
+          member_it->udt = udt;
+          member_it->stream_end = symrec.end;
+          is_inited = off > 0 && symrec.size > 0 && symrec.type == PDB_LF_FIELDLIST;
+        }
+      }
+    }
+  }
+
+  return is_inited;
+}
+
+PDB_API syms_bool
+pdb_member_it_next(pdb_member_it *member_it, pdb_member *member)
+{
+  U16 lf_kind;
+  pdb_stream *stream = &member_it->stream;
+  pdb_uint name_size;
+  syms_bool redone;
+  syms_bool is_member_valid = syms_false;
+
+redo:;
+  redone = syms_false;
+
+  if (member_it->stream.off + sizeof(lf_kind) >= member_it->stream_end) {
+    SYMS_ASSERT_PARANOID(member_it->stream.off == member_it->stream_end);
+    return syms_false;
+  }
+
+  if ( ! pdb_stream_read_u16(stream, &lf_kind)) {
+    SYMS_ASSERT_FAILURE_PARANOID("unable to parse LF_ARGLIST -- size of leaves doesnt add up");
+    return syms_false;
+  }
+
+  switch (lf_kind) {
+  case PDB_LF_MEMBER: {
+    pdb_lf_member lf;
+
+    if (pdb_stream_read_struct(stream, &lf, pdb_lf_member)) {
+      member->type = PDB_MEMBER_TYPE_DATA;
+      member->itype = lf.itype;
+      member->attr = lf.attr;
+      pdb_stream_read_numeric_u32(stream, &member->u.data.offset);
+
+      member->name = pdb_pointer_bake_stream_str(stream);
+      name_size = pdb_stream_read_str(stream, 0, 0);
+      pdb_stream_skip(stream, name_size);
+    }
+  } break;
+
+  case PDB_LF_STMEMBER: {
+    pdb_lf_stmember lf;
+
+    if (pdb_stream_read_struct(stream, &lf, pdb_lf_stmember)) {
+      member->type = PDB_MEMBER_TYPE_STATIC_DATA;
+      member->itype = lf.index;
+      member->attr = lf.attr;
+
+      member->name = pdb_pointer_bake_stream_str(stream);
+      name_size = pdb_stream_read_str(stream, 0, 0);
+      pdb_stream_skip(stream, name_size);
+    }
+  } break;
+
+  case PDB_LF_METHOD: {
+    pdb_lf_method lf;
+
+    if (pdb_stream_read_struct(&member_it->stream, &lf, pdb_lf_method)) {
+      member->type = PDB_MEMBER_TYPE_METHODLIST;
+      member->itype = lf.itype_list; /* method has itype that points to an array of ml_method */
+      member->u.methodlist.count = lf.count; /* why counter is here and not part of ml_methodlist?! */
+
+      member->name = pdb_pointer_bake_stream_str(stream);
+      name_size = pdb_stream_read_str(stream, 0, 0);
+      pdb_stream_skip(stream, name_size);
+    }
+  } break;
+
+  case PDB_LF_ONEMETHOD: {
+    pdb_lf_onemethod lf;
+
+    if (pdb_stream_read_struct(stream, &lf, pdb_lf_onemethod)) {
+      pdb_uint mprop = PDB_CV_FLDATTR_MPROP_MASK(lf.attr);
+      pdb_uint vbaseoff = 0;
+
+      if (mprop == PDB_CV_FLDATTR_MPROP_PUREINTRO || 
+          mprop == PDB_CV_FLDATTR_MPROP_INTRO) {
+        pdb_stream_read_u32(stream, &vbaseoff);
+      }
+
+      member->type = PDB_MEMBER_TYPE_METHOD;
+      member->itype = lf.itype;
+      member->attr = lf.attr;
+      member->u.method.vbaseoff = vbaseoff;
+
+      member->name = pdb_pointer_bake_stream_str(stream);
+      name_size = pdb_stream_read_str(stream, 0, 0);
+      pdb_stream_skip(stream, name_size);
+    }
+  } break;
+
+  case PDB_LF_ENUMERATE: {
+    pdb_lf_enumerate lf;
+
+    if (pdb_stream_read_struct(stream, &lf, pdb_lf_enumerate)) {
+      member->type = PDB_MEMBER_TYPE_ENUMERATOR;
+      member->attr = lf.attr;
+      pdb_stream_read_numeric(&member_it->stream, &member->u.enumerator.value);
+      member->itype = member->u.enumerator.value.itype;
+
+      member->name = pdb_pointer_bake_stream_str(stream);
+      name_size = pdb_stream_read_str(stream, 0, 0);
+      pdb_stream_skip(stream, name_size);
+    }
+  } break;
+
+  case PDB_LF_NESTTYPE: {
+    pdb_lf_nesttype lf;
+
+    if (pdb_stream_read_struct(stream, &lf, pdb_lf_nesttype)) {
+      member->type = PDB_MEMBER_TYPE_NESTED_TYPE;
+      member->itype = lf.itype;
+
+      member->name = pdb_pointer_bake_stream_str(stream);
+      name_size = pdb_stream_read_str(stream, 0, 0);
+      pdb_stream_skip(stream, name_size);
+    }
+  } break;
+
+  case PDB_LF_BCLASS: {
+    pdb_lf_bclass lf;
+
+    if (pdb_stream_read_struct(stream, &lf, pdb_lf_bclass)) {
+      member->type = PDB_MEMBER_TYPE_BASE_CLASS;
+      member->itype = lf.itype;
+      pdb_stream_read_numeric_u32(stream, &member->u.base_class.offset);
+
+      member->name = pdb_pointer_bake_str(syms_string_init_lit("BaseClass"));
+    }
+  } break;
+
+  case PDB_LF_VBCLASS:
+  case PDB_LF_IVBCLASS: {
+    pdb_lf_vbclass lf;
+    pdb_numeric num;
+
+    pdb_stream_read_struct(stream, &lf, pdb_lf_vbclass);
+    pdb_stream_read_numeric(stream, &num);
+  } break;
+
+  case PDB_LF_VFUNCTAB: {
+    pdb_lf_vfunctab lf;
+  
+    if (pdb_stream_read_struct(stream, &lf, pdb_lf_vfunctab)) {
+      member->type = PDB_MEMBER_TYPE_VIRTUAL_TABLE;
+      member->itype = lf.itype;
+      member->name = pdb_pointer_bake_str(syms_string_init_lit("vpftable"));
+      member->u.vtab.offset = 0;
+    }
+  } break;
+    
+  case PDB_LF_VFUNCOFF: {
+    pdb_lf_vfuncoff lf;
+
+    if (pdb_stream_read_struct(stream, &lf, pdb_lf_vfuncoff)) {
+      member->type = PDB_MEMBER_TYPE_VIRTUAL_TABLE;
+      member->itype = lf.itype;
+      member->name = pdb_pointer_bake_str(syms_string_init_lit("vpftable"));
+      member->u.vtab.offset = lf.offset;
+    }
+  } break;
+
+  case PDB_LF_INDEX: {
+    /* leaf represents a reference to a field list */
+    pdb_lf_index lf;
+    pdb_uint off;
+    pdb_symrec symrec;
+
+    member->name = pdb_pointer_bake_str(syms_string_init_lit("BAD_LF_INDEX"));
+    member->type = PDB_MEMBER_TYPE_NULL;
+    member->itype = 0;
+    member->attr = 0;
+
+    if (!pdb_stream_read_struct(stream, &lf, pdb_lf_index)) 
+      break;
+    if (!pdb_tm_get_itype_offset(member_it->pdb, lf.itype, &off)) 
+      break;
+    if (!pdb_stream_seek(stream, off)) 
+      break;
+    if (!pdb_stream_read_symrec(stream, &symrec)) 
+      break;
+    if (symrec.size == 0 || symrec.type != PDB_LF_FIELDLIST) 
+      break;
+
+    member_it->stream_end = symrec.end;
+
+    if (!redone) {
+      redone = syms_true;
+      goto redo;
+    }
+  } break;
+
+  default: {
+    SYMS_ASSERT_FAILURE("Unknown field type");
+  } break;
+  }
+
+  /* make sure stream is aligned to 4 bytes, 
+   * stream contains variable-length data and 
+   * offsets of leaves are aligned to four bytes
+   * for memory reasons. */
+  pdb_stream_align(&member_it->stream, 4);
+
+  is_member_valid = PDB_STREAM_NO_ERROR(stream->flags);
+  return is_member_valid;
+}
+
+PDB_API syms_bool
+pdb_line_it_init(pdb_line_it *line_it, pdb_mod *mod)
+{
+  syms_bool result;
+
+  syms_memzero(line_it, sizeof(*line_it));
+  result = syms_false;
+  line_it->mod = *mod;
+  line_it->last_read_ln = 0;
+  if (pdb_mod_get_debug_sec(mod, PDB_MOD_SEC_LINES_C13, &line_it->stream)) {
+    pdb_stream_init_null(&line_it->u.c13.cv_lines);
+    line_it->format = PDB_LINE_FORMAT_C13;
+    line_it->u.c13.sec_end = 0;
+    result = syms_true;
+  } else if (pdb_mod_get_debug_sec(mod, PDB_MOD_SEC_LINES_C11, &line_it->stream)) {
+    line_it->format = PDB_LINE_FORMAT_C11;
+    line_it->u.c11.file_count = 0;
+    result = syms_true;
+  }
+  return result;
+}
+PDB_API syms_bool
+pdb_line_it_next_c13(pdb_line_it *line_it, pdb_uint *off_out, pdb_uint *ln_out)
+{
+  pdb_stream *line_data = &line_it->u.c13.cv_lines;
+  syms_bool is_line_valid = syms_false;
+  syms_bool fetch_next_section;
+
+  do {
+    pdb_cv_line cv_line;
+
+    fetch_next_section = syms_false;
+
+    while (pdb_stream_is_null(line_data) || PDB_STREAM_READ_OR_SEEK_FAILED(line_data->flags)) {
+      SYMS_ASSERT_PARANOID(~line_it->flags & PDB_LINE_IT_FLAGS_NEW_SECTION);
+
+      if (line_it->stream.off >= line_it->stream.size) {
+        return syms_false;
+      }
+
+      while (line_it->stream.off < line_it->stream.size) {
+        pdb_uint type = PDB_UINT_MAX;
+        pdb_uint size = PDB_UINT_MAX;
+
+        if (PDB_STREAM_READ_OR_SEEK_FAILED(line_it->stream.flags)) {
+          return syms_false;
+        }
+
+        pdb_stream_read_u32(&line_it->stream, &type);
+        pdb_stream_read_u32(&line_it->stream, &size);
+
+        if (type == PDB_CV_SS_TYPE_LINES && !PDB_CV_SS_TYPE_IGNORE(type)) {
+          pdb_cv_line_sec sec;
+          pdb_cv_src_file file;
+
+          *line_data = pdb_stream_subset(&line_it->stream, line_it->stream.off, size);
+          pdb_stream_skip(&line_it->stream, size);
+
+          file.chksum_off = SYMS_UINT32_MAX;
+          file.num_lines = 0;
+          file.lines_size = 0;
+
+          pdb_stream_read_struct(line_data, &sec,  pdb_cv_line_sec);
+          pdb_stream_read_struct(line_data, &file, pdb_cv_src_file);
+
+          if (PDB_STREAM_READ_OR_SEEK_FAILED(line_data->flags)) {
+            continue;
+          }
+
+          if (pdb_mod_infer_fileid(&line_it->mod, file.chksum_off, &line_it->fi)) {
+            line_it->flags |= PDB_LINE_IT_FLAGS_NEW_SECTION;
+
+            line_it->sec        = sec.sec;
+            line_it->sec_off    = sec.sec_off;
+            line_it->sec_size   = sec.len;
+
+            line_it->u.c13.line_index_max = file.num_lines; 
+          } else {
+            /* when offset is invalid we are unable to infer file names and checksums.
+             * move to next line section. */
+            SYMS_ASSERT_PARANOID("invalid checksum section offset");
+            pdb_stream_init_null(line_data);
+          }
+          break;
+        } 
+        pdb_stream_skip(&line_it->stream, size);
+      }
+    }
+
+    if (pdb_stream_read_struct(line_data, &cv_line, pdb_cv_line)) {
+      *off_out = line_it->sec_off + cv_line.off;
+      *ln_out  = PDB_CV_LINE_GET_LN(cv_line);
+      if (*ln_out == PDB_UINT_MAX) {
+        *ln_out = 0;
+      }
+      is_line_valid = syms_true;
+    } else {
+      if (!PDB_STREAM_READ_OR_SEEK_FAILED(line_it->stream.flags)) {
+        fetch_next_section = syms_true;
+      }
+    }
+  } while (fetch_next_section);
+
+  return is_line_valid;
+}
+PDB_API syms_bool
+pdb_line_it_next_c11(pdb_line_it *line_it, pdb_uint *off_out, pdb_uint *ln_out)
+{
+  syms_bool is_line_valid = syms_false;
+
+  if (line_it->u.c11.file_count == 0) {
+    U16 file_count = 0;
+    U16 range_count = 0;
+
+    SYMS_ASSERT(~line_it->flags & PDB_LINE_IT_FLAGS_NEW_SECTION);
+    pdb_stream_read_u16(&line_it->stream, &file_count);
+    pdb_stream_read_u16(&line_it->stream, &range_count);
+    pdb_stream_skip(&line_it->stream, sizeof(pdb_uint)*file_count);
+    pdb_stream_skip(&line_it->stream, sizeof(pdb_uint)*2*range_count);
+    pdb_stream_skip(&line_it->stream, sizeof(U16)*range_count);
+    pdb_stream_align(&line_it->stream, 4);
+    if (PDB_STREAM_READ_OR_SEEK_FAILED(line_it->stream.flags)) {
+      return syms_false;
+    }
+
+    line_it->u.c11.file_index = 0;
+    line_it->u.c11.file_count = file_count;
+    line_it->flags |= PDB_LINE_IT_FLAGS_NEW_SECTION;
+  }
+
+  for (;;) {
+    if (line_it->flags & PDB_LINE_IT_FLAGS_NEW_SECTION) {
+      U16 pad;
+      pdb_uint path_size;
+      pdb_pointer path = pdb_pointer_bake_null();
+
+      if (line_it->u.c11.file_index >= line_it->u.c11.file_count) {
+        return syms_false;
+      }
+
+      pdb_stream_read_u16(&line_it->stream, &line_it->u.c11.filesec_count);
+      pdb_stream_read_u16(&line_it->stream, &pad);
+      pdb_stream_skip(&line_it->stream, sizeof(pdb_uint)*line_it->u.c11.filesec_count);
+      line_it->u.c11.secrange_stream = pdb_stream_subset(&line_it->stream, line_it->stream.off, sizeof(pdb_uint)*2*line_it->u.c11.filesec_count);
+      pdb_stream_skip(&line_it->stream, sizeof(pdb_uint)*2*line_it->u.c11.filesec_count);
+      if (PDB_STREAM_READ_OR_SEEK_FAILED(line_it->stream.flags)) {
+        return syms_false;
+      }
+
+      path_size = 0;
+      switch (line_it->stream.pdb->ver) {
+      case PDB_VER_VC50:
+      case PDB_VER_VC4:
+      case PDB_VER_VC2:
+      case PDB_VER_VC98: {
+        pdb_stream_read_u08(&line_it->stream, (U8 *)&path_size);
+        path = pdb_pointer_bake(&line_it->stream, path_size);
+      } break;
+      case PDB_VER_VC70:
+      case PDB_VER_VC70_DEP:
+      case PDB_VER_VC80:
+      case PDB_VER_VC140:
+      case PDB_VER_VC110: {
+        path = pdb_pointer_bake_stream_str(&line_it->stream);
+        path_size = pdb_pointer_get_size(&path);
+        path_size += 1;
+      } break;
+      }
+
+      pdb_stream_skip(&line_it->stream, path_size);
+      pdb_stream_align(&line_it->stream, 4);
+
+      line_it->fi.path = path;
+      line_it->fi.chksum_type = PDB_CV_CHECKSUM_NULL; /* NOTE(nick): c11 doesn't have checksums */
+      line_it->fi.chksum = pdb_pointer_bake_null();
+      if (PDB_STREAM_READ_OR_SEEK_FAILED(line_it->stream.flags)) {
+        return syms_false;
+      }
+
+      line_it->u.c11.file_index += 1;
+      line_it->u.c11.filesec_index = 0;
+      line_it->u.c11.pair_index = 1;
+      line_it->u.c11.pair_count = 0;
+    }
+
+    if (line_it->u.c11.pair_index >= line_it->u.c11.pair_count) {
+      pdb_uint off_size, ln_size;
+      pdb_uint sec_lo, sec_hi;
+
+      if (line_it->u.c11.filesec_index >= line_it->u.c11.filesec_count) {
+        line_it->flags |= PDB_LINE_IT_FLAGS_NEW_SECTION;
+        continue;
+      }
+      pdb_stream_read_u16(&line_it->stream, &line_it->sec);
+      pdb_stream_read_u16(&line_it->stream, &line_it->u.c11.pair_count);
+      if (PDB_STREAM_READ_OR_SEEK_FAILED(line_it->stream.flags)) {
+        break;
+      }
+      pdb_stream_read_u32(&line_it->u.c11.secrange_stream, &sec_lo);
+      pdb_stream_read_u32(&line_it->u.c11.secrange_stream, &sec_hi);
+      if (PDB_STREAM_READ_OR_SEEK_FAILED(line_it->u.c11.secrange_stream.flags)) {
+        break;
+      }
+      line_it->sec_off  = sec_lo;
+      line_it->sec_size = sec_hi - sec_lo;
+
+      off_size = line_it->u.c11.pair_count * sizeof(pdb_uint);
+      ln_size = line_it->u.c11.pair_count * sizeof(U16);
+      line_it->u.c11.off_stream = pdb_stream_subset(&line_it->stream, line_it->stream.off, off_size);
+      line_it->u.c11.ln_stream  = pdb_stream_subset(&line_it->stream, line_it->stream.off + off_size, ln_size);
+      pdb_stream_skip(&line_it->stream, off_size + ln_size);
+      pdb_stream_align(&line_it->stream, 4);
+
+      line_it->u.c11.pair_index = 0;
+      line_it->u.c11.filesec_index += 1;
+    }
+
+    if (line_it->u.c11.pair_index < line_it->u.c11.pair_count) {
+      pdb_uint off = 0;
+      U16 ln = 0;
+
+      pdb_stream_read_u32(&line_it->u.c11.off_stream, &off);
+      pdb_stream_read_u16(&line_it->u.c11.ln_stream, &ln);
+
+      *off_out = off;
+      *ln_out  = ln;
+
+      line_it->u.c11.pair_index += 1;
+      is_line_valid = syms_true;
+      break;
+    }
+#if 0
+    else if (line_it->u.c11.pair_index == line_it->u.c11.pair_count) {
+      *off_out = line_it->sec_off + line_it->sec_size;
+      *ln_out = line_it->last_read_ln + 1;
+      line_it->u.c11.pair_index += 1;
+      is_line_valid = syms_true;
+    }
+#endif
+
+    line_it->flags |= PDB_LINE_IT_FLAGS_NEW_SECTION;
+  }
+
+  return is_line_valid;
+}
+PDB_API syms_bool
+pdb_line_it_next(pdb_line_it *line_it, pdb_line *line_out)
+{
+  syms_bool is_line_valid = syms_false;
+  pdb_uint off = PDB_UINT_MAX, ln = 0;
+  line_it->flags &= ~PDB_LINE_IT_FLAGS_NEW_SECTION;
+  switch (line_it->format) {
+  case PDB_LINE_FORMAT_NULL:  is_line_valid = syms_false; break;
+  case PDB_LINE_FORMAT_C11:   is_line_valid = pdb_line_it_next_c11(line_it, &off, &ln); break;
+  case PDB_LINE_FORMAT_C13:   is_line_valid = pdb_line_it_next_c13(line_it, &off, &ln); break;
+  }
+  if (is_line_valid) {
+    line_it->last_read_ln = ln;
+    line_out->sec = line_it->sec;
+    line_out->off = off;
+    line_out->ln = ln;
+  }
+  return is_line_valid;
+}
+
+PDB_API syms_bool
+pdb_src_to_va(pdb_context *pdb, char *filename, pdb_uint filename_len, pdb_uint ln, pdb_map *map)
+{
+  pdb_mod_it mod_it;
+  pdb_mod mod;
+  pdb_stream strtable;
+  pdb_stream lines;
+  pdb_uint closest_line_off, closest_ln, closest_max_ln, closest_chksum;
+  pdb_cv_line_sec closest_sec;
+  SymsString filename_str;
+
+  filename_str = syms_string_init(filename, filename_len);
+  /* Check if this file is present in the PDB */
+  if (!pdb_strtable_find_str(pdb, filename_str, PDB_STRCMP_FLAG_NOCASE, 0)) return syms_false;
+
+  strtable = pdb_get_strtable(pdb);
+  if (!pdb_mod_it_init(&mod_it, pdb)) return syms_false;
+
+  closest_line_off = 0; 
+  closest_max_ln = 0; 
+  closest_chksum = 0;
+  closest_ln = PDB_UINT_MAX;
+  syms_memset(&closest_sec, 0, sizeof(closest_sec));
+
+  while (pdb_mod_it_next(&mod_it, &mod)) {
+    pdb_uint chksum_base;
+    pdb_uint sec_type, sec_size, sec_end;
+
+    if (!pdb_mod_get_debug_sec(&mod, PDB_MOD_SEC_LINES_C13, &lines)) {
+      continue;
+    }
+
+    /* Looking up the checksum section. We need it in order to 
+     * resolve file names that are refed in the lines section. */
+    chksum_base = 0;
+    for (;;) {
+      if (!pdb_stream_read_u32(&lines, &sec_type)) break;
+      if (!pdb_stream_read_u32(&lines, &sec_size)) break;
+      if (sec_type == PDB_CV_SS_TYPE_FILE_CHKSUM) {
+        chksum_base = lines.off;
+        break;
+      }
+      if (!pdb_stream_skip(&lines, sec_size)) break;
+    }
+
+    if (!chksum_base) continue;
+    if (!pdb_stream_seek(&lines, 0)) continue;
+
+    for (;;) {
+      pdb_cv_line_sec section;
+      pdb_cv_src_file file;
+      pdb_cv_line line;
+      pdb_uint min_line, max_line;
+      pdb_uint i;
+
+      if (!pdb_stream_read_u32(&lines, &sec_type)) break;
+      if (!pdb_stream_read_u32(&lines, &sec_size)) break;
+      sec_end = lines.off + sec_size;
+
+      if (sec_type != PDB_CV_SS_TYPE_LINES) goto next_sec;
+      if (!pdb_stream_read_struct(&lines, &section, pdb_cv_line_sec)) break;
+      if (!pdb_stream_read_struct(&lines, &file, pdb_cv_src_file))    break;
+
+      {
+        /* Resolving file name that is refed by this section
+         * and compare it to the specified file name. */
+        pdb_uint lines_start = lines.off;
+        pdb_cv_file_checksum chksum;
+
+        file.chksum_off += chksum_base;
+        if (!pdb_stream_seek(&lines, file.chksum_off)) goto next_sec;
+        if (!pdb_stream_read_struct(&lines, &chksum, pdb_cv_file_checksum)) goto next_sec;
+
+        if (!pdb_stream_seek(&strtable, chksum.name_off)) goto next_sec;
+        if (!pdb_strcmp_stream_nocase(filename_str, &strtable)) goto next_sec;
+        if (!pdb_stream_seek(&lines, lines_start)) break;
+      }
+
+      { 
+        /* Range check specified line number with this section. */
+        pdb_uint lines_start = lines.off;
+        pdb_uint last_line_off;
+
+        if (!pdb_stream_read_struct(&lines, &line, pdb_cv_line)) goto next_sec;
+        min_line = PDB_CV_LINE_GET_LN(line);
+
+        if (file.num_lines > 1) {
+          last_line_off = (file.num_lines - 2)*sizeof(pdb_cv_line);
+          if (!pdb_stream_skip(&lines, last_line_off)) goto next_sec;
+          if (!pdb_stream_read_struct(&lines, &line, pdb_cv_line)) goto next_sec;
+          max_line = PDB_CV_LINE_GET_LN(line);
+        } else {
+          max_line = min_line + 1;
+        }
+
+        SYMS_ASSERT(min_line < max_line);
+        if (ln < min_line || ln > max_line) goto next_sec;
+        if (!pdb_stream_seek(&lines, lines_start)) return syms_false;
+      }
+
+      for (i = 0; i < file.num_lines; ++i) {
+        pdb_uint t = lines.off;
+        if (!pdb_stream_read_struct(&lines, &line, pdb_cv_line)) return syms_false;
+        if (PDB_CV_LINE_GET_LN(line) <= closest_ln && PDB_CV_LINE_GET_LN(line) >= ln) {
+          closest_ln = PDB_CV_LINE_GET_LN(line);
+          closest_line_off = t;
+          closest_max_ln = max_line;
+          closest_chksum = file.chksum_off;
+          syms_memcpy(&closest_sec, &section, sizeof(section));
+          if (closest_ln == ln) goto exit;
+        }
+      }
+      if (closest_line_off) goto exit;
+
+next_sec:
+      if (!pdb_stream_seek(&lines, sec_end)) {
+        SYMS_ASSERT_CORRUPTED_STREAM;
+        return syms_false;
+      }
+    }
+  }
+
+exit:
+  if (!closest_line_off) return syms_false;
+
+  if (map) {
+    pdb_cv_line line;
+
+    if (!pdb_stream_seek(&lines, closest_line_off))    return syms_false;
+    if (!pdb_stream_read(&lines, &line, sizeof(line))) return syms_false;
+    if (closest_sec.sec >= pdb->dbi.secs_num)          return syms_false;
+
+    map->sec     = closest_sec.sec;
+    map->sec_off = closest_sec.sec_off + line.off;
+
+    { /* Calc instruction block byte length */
+      pdb_cv_line next_line;
+
+      next_line.off = 0;
+      if (closest_ln < closest_max_ln) {
+        while (pdb_stream_read_struct(&lines, &next_line, pdb_cv_line)) {
+          if (PDB_CV_LINE_GET_LN(next_line) != closest_ln)
+            break;
+          next_line.off = 0;
+        }
+      }
+      if (!next_line.off)
+        next_line.off = closest_sec.len;
+      SYMS_ASSERT(next_line.off > line.off);
+      map->size = next_line.off - line.off;
+    }
+
+    map->ln = PDB_CV_LINE_GET_LN(line);
+
+    if (!pdb_mod_infer_fileid(&mod, closest_chksum, &map->fi)) {
+      SYMS_INVALID_CODE_PATH;
+    }
+  }
+
+  return syms_true;
+}
+
+PDB_API pdb_stream
+pdb_get_strtable(pdb_context *pdb)
+{
+  pdb_stream result = pdb->strtable;
+  pdb_stream_seek(&result, 0);
+  return result;
+}
+
+PDB_API pdb_stream
+pdb_get_stroffs(pdb_context *pdb)
+{
+  pdb_stream result = pdb->stroffs;
+  pdb_stream_seek(&result, 0);
+  return result;
+}
+
+PDB_API syms_bool
+pdb_strtable_off_to_str(pdb_context *pdb, pdb_stroff stroff, pdb_pointer *str_out)
+{
+  syms_bool done = syms_false;
+  pdb_stream strtable = pdb_get_strtable(pdb);
+  if (pdb_stream_seek(&strtable, stroff)) {
+    *str_out = pdb_pointer_bake_stream_str(&strtable);
+    done = syms_true;
+  }
+  return done;
+}
+
+PDB_API syms_bool
+pdb_strtable_at(pdb_context *pdb, pdb_ni ni, pdb_pointer *str_out)
+{
+  pdb_stream strtable = pdb_get_strtable(pdb);
+  pdb_stream stroffs = pdb_get_stroffs(pdb);
+  if (pdb_stream_seek(&stroffs, ni*sizeof(pdb_stroff))) {
+    pdb_stroff strtable_off;
+    if (pdb_stream_read(&stroffs, &strtable_off, sizeof(strtable_off))) {
+      if (pdb_stream_seek(&strtable, strtable_off)) {
+        *str_out = pdb_pointer_bake_stream_str(&strtable);
+        return syms_true;
+      }
+    }
+  }
+  return syms_false;
+}
+
+PDB_API syms_bool
+pdb_strtable_find(pdb_context *pdb, pdb_pointer *name, pdb_strcmp_flags_e cmp_flags, pdb_ni *ni_out)
+{
+  pdb_stream strtable = pdb_get_strtable(pdb);
+  pdb_stream stroffs = pdb_get_stroffs(pdb);
+
+  pdb_uint offsets_num = pdb->stroffs.size/sizeof(pdb_stroff);
+  pdb_uint hash = pdb_hashV1_pointer(pdb, name, offsets_num);
+  pdb_uint indx = hash;
+
+  do {
+    pdb_stroff stroff;
+
+    if (!pdb_stream_seek(&stroffs, indx*sizeof(pdb_stroff)))  break;
+    if (!pdb_stream_read(&stroffs, &stroff, sizeof(stroff))) break;
+    if (!pdb_stream_seek(&strtable, stroff))                  break;
+
+    if (pdb_stream_strcmp_pointer_(&strtable, name, cmp_flags)) {
+      if (ni_out) {
+        *ni_out = indx;
+      }
+      return syms_true;
+    }
+
+    indx = (indx + 1) % offsets_num;
+  } while (indx != hash);
+  return syms_false;
+}
+
+PDB_API syms_bool
+pdb_strtable_find_str(pdb_context *pdb, SymsString name, pdb_strcmp_flags_e cmp_flags, pdb_ni *ni)
+{
+  pdb_pointer name_ptr = pdb_pointer_bake_str(name);
+  return pdb_strtable_find(pdb, &name_ptr, cmp_flags, ni);
+}
+
+PDB_API S32
+pdb_addr_map_cmp(pdb_uint isect0, pdb_uint off0, pdb_uint isect1, pdb_uint off1)
+{
+  S32 result;
+  if (isect0 == isect1) {
+    result = (S32)off0 - (S32)off1;
+  } else {
+    result = (S32)isect0 - (S32)isect1;
+  }
+  return result;
+}
+
+PDB_API syms_bool
+pdb_psi_read_pubsym(pdb_stream *addr_stream, pdb_stream *sym_stream, S32 index, pdb_cv_datasym32 *pubsym_out)
+{
+  pdb_uint sym_off;
+  pdb_symrec sym_header;
+
+  if (index < 0) return syms_false;
+
+  if (!pdb_stream_seek(addr_stream, (pdb_uint)index*sizeof(pdb_uint))) return syms_false;
+  if (!pdb_stream_read_u32(addr_stream, &sym_off))           return syms_false;
+
+  if (!pdb_stream_seek(sym_stream, sym_off))            return syms_false;
+  if (!pdb_stream_read_symrec(sym_stream, &sym_header)) return syms_false;
+  
+  if (sym_header.type != PDB_CV_SYM_PUB32_ST && sym_header.type != PDB_CV_SYM_PUB32) return syms_false;
+  if (!pdb_stream_read(sym_stream, pubsym_out, sizeof(*pubsym_out)))                 return syms_false;
+
+  return syms_true;
+}
+
+PDB_API syms_bool
+pdb_find_nearest_sym(pdb_context *pdb, SymsAddr va, pdb_pointer *name_out)
+{
+  /* TODO(nick; Oct 6 2019): This procedure is a first draft and it needs a refactor. */
+
+  syms_bool inited = syms_false;
+  pdb_psi_header psi_header;
+  pdb_stream psi_stream;
+  pdb_stream sym_stream;
+  /* NOTE(nick): addr_stream contains 32bit offsets into the sym_stream.
+   * Offsets are sorted and point to symbols of type PDB_CV_SYM_PUB* */
+  pdb_stream addr_stream;
+  pdb_uint addr_map_num = 0;
+
+  pdb_isec sec;
+  pdb_isec_umm off;
+
+  S32 min, max;
+  S32 pubsym_index;
+  pdb_cv_datasym32 pubsym;
+
+  if (!pdb_build_sec_off(pdb, va, &sec, &off)) return syms_false;
+
+  syms_memset(&psi_header, 0, sizeof(psi_header));
+
+  if (pdb_stream_init(pdb, pdb->dbi.pubsym_sn, &psi_stream)) {
+    if (pdb_stream_read(&psi_stream, &psi_header, sizeof(pdb_psi_header))) {
+      if (psi_header.addr_map_size > 0) {
+        pdb_uint addr_table_base = sizeof(pdb_psi_header) + psi_header.sym_hash_size;
+        addr_map_num = psi_header.addr_map_size/sizeof(pdb_uint);
+        addr_stream = psi_stream;
+
+        addr_stream = pdb_stream_subset(&psi_stream, addr_table_base, psi_header.addr_map_size);
+        if (addr_stream.size > 0) {
+          if (pdb_stream_init(pdb, pdb->dbi.symrec_sn, &sym_stream)) {
+            inited = syms_true;
+          }
+        }
+      }
+    }
+  }
+  SYMS_ASSERT(inited);
+  if (!inited) return syms_false;
+
+  min = 0;
+  max = (S32)(psi_header.addr_map_size / sizeof(S32)) - 1;
+  while (min < max) {
+    S32 mid;
+    S32 cmp;
+
+    mid = min + (max - min + 1) / 2;
+    if (!pdb_psi_read_pubsym(&addr_stream, &sym_stream, mid, &pubsym)) {
+      break;
+    }
+
+    cmp = pdb_addr_map_cmp(sec, off, pubsym.sec, pubsym.sec_off);
+    if (cmp < 0) {
+      max = mid - 1;
+    } else if (cmp > 0) {
+      min = mid;
+    } else {
+      min = max = mid;
+    }
+  }
+
+  pubsym_index = min;
+  if (!pdb_psi_read_pubsym(&addr_stream, &sym_stream, pubsym_index, &pubsym)) {
+    return syms_false;
+  }
+
+  if (pubsym.sec == sec) {
+    /* NOTE(nick): When ICF(Identical Code Folding) is turned on linker
+     * might output procedures that have same address, in this case reference
+     * implementation returns first element from the list */
+    min = pubsym_index - 1;
+
+    while (pubsym_index >= 0 && min > 0) {
+      pdb_cv_datasym32 pubsym_min;
+      S32 cmp;
+
+      if (!pdb_psi_read_pubsym(&addr_stream, &sym_stream, pubsym_index, &pubsym)) {
+        break;
+      }
+
+      if (!pdb_psi_read_pubsym(&addr_stream, &sym_stream, min, &pubsym_min)) {
+        break;
+      }
+
+      cmp = pdb_addr_map_cmp(pubsym_min.sec, pubsym_min.sec_off, pubsym.sec, pubsym.sec_off);
+      if (cmp != 0) {
+        break;
+      }
+
+      SYMS_ASSERT(min > 0);
+      min -= 1;
+
+      SYMS_ASSERT(pubsym_index > 0);
+      pubsym_index -= 1;
+    }
+  } else {
+    /* NOTE(nick): Comment from reference implementation:
+     * Boundary conditions.
+     * Example: given publics at (a=1:10, b=1:20, c=2:10, d=2:20, e=4:0),
+     * search for (1: 9) returns (a,-1)
+     *        for (1:11) returns (a,1)
+     *        for (1:21) returns (b,1)
+     *        for (2: 9) returns (c,-1)
+     *        for (2:11) returns (c,1)
+     *        for (2:21) returns (d,1)
+     *        for (3:xx) returns NULL
+     * so, for cases (2:9), we must advance ppsymLo from (1:21) to (2:9)
+     *
+     * Need to use a loop instead of justing returning the next symbol,
+     * due to ICF, ppsymLo may not point to the last among symbols
+     * having same address. */
+
+    while (pubsym.sec < sec) {
+      pubsym_index += 1;
+
+      if (pubsym_index >= (S32)addr_map_num)  {
+        return syms_false;
+      }
+
+      if (!pdb_psi_read_pubsym(&addr_stream, &sym_stream, pubsym_index, &pubsym)) {
+        return syms_false;
+      }
+
+      if (pubsym.sec > sec) {
+        return syms_false;
+      }
+    }
+  }
+
+  if (!pdb_psi_read_pubsym(&addr_stream, &sym_stream, pubsym_index, &pubsym)) return syms_false;
+  *name_out = pdb_pointer_bake_stream_str(&sym_stream);
+  return syms_true;
+}
+
+PDB_API syms_bool
+pdb_find_nearest_sc(pdb_context *pdb, SymsAddr va, pdb_sc *sc_out)
+{
+  pdb_isec sec;
+  pdb_isec_umm off;
+  syms_bool result = syms_false;
+  if (pdb_build_sec_off(pdb, va, &sec, &off)) {
+    result = pdb_find_nearest_sc_ex(pdb, sec, off, sc_out);
+  }
+  return result;
+}
+
+PDB_API syms_bool
+pdb_find_nearest_sc_ex(pdb_context *pdb, pdb_uint sec, pdb_uint off, pdb_sc *sc_out)
+{
+  pdb_stream dbi_stream;
+  pdb_uint version;
+
+  if (!pdb_stream_init(pdb, PDB_DEFAULT_STREAM_DBI, &dbi_stream)) {
+    return syms_false;
+  }
+
+  dbi_stream = pdb_stream_subset(&dbi_stream, pdb->dbi.seccon_off, pdb->dbi.seccon_len);
+  if (!pdb_stream_read_u32(&dbi_stream, &version)) {
+    return syms_false;
+  }
+
+  switch (version) {
+  case PDB_DBI_SC_VER_60: {
+    pdb_uint sc_count = (pdb->dbi.seccon_len - sizeof(version)) / sizeof(pdb_sc40);
+    if (sc_count > 0) {
+      pdb_uint min = 0;
+      pdb_uint max = sc_count - 1;
+
+      while (min < max) {
+        pdb_uint mid = min + (max - min + 1) / 2;
+        pdb_sc40 sc;
+        pdb_int cmp;
+
+        if (!pdb_stream_seek(&dbi_stream, sizeof(version) + mid * sizeof(sc))) {
+          break;
+        }
+        if (!pdb_stream_read_struct(&dbi_stream, &sc, pdb_sc40)) {
+          break;
+        }
+
+        if (sec == sc.sec && off >= sc.sec_off && off < sc.sec_off + sc.size) {
+          cmp = 0;
+        } else {
+          cmp = pdb_addr_map_cmp(sec, off, sc.sec, sc.sec_off);
+        }
+
+        if (cmp < 0) {
+          max = mid - 1;
+        } else if (cmp > 0) {
+          min = mid;
+        } else {
+          min = max = mid;
+          
+          syms_memset(sc_out, 0, sizeof(*sc_out));
+          sc_out->sec = sc.sec;
+          sc_out->sec_off = sc.sec_off;
+          sc_out->size = sc.size;
+          sc_out->flags = sc.flags;
+          sc_out->imod = sc.imod;
+
+          return syms_true;
+        }
+      }
+    }
+  } break;
+
+  case PDB_DBI_SC_VER_2: {
+    /* TODO(nick): implement */
+  } break;
+  }
+
+  return syms_false;
+}
+
+PDB_API pdb_uint
+pdb_parse_gsi_hash_table(pdb_stream *gsi_stream, 
+                         struct SymsArena *arena,
+                         pdb_gsi_hr ***out_table, pdb_uint *out_num_table)
+{
+#define PDB_GSI_HASH_TABLE_MAX   4096u
+#define PDB_GSI_BITMAP_SIZE      ((((PDB_GSI_HASH_TABLE_MAX + 1) + 31)/32*32)/8)
+#define PDB_GSI_BITMAP_U32_COUNT (PDB_GSI_BITMAP_SIZE/sizeof(pdb_uint))
+
+  pdb_gsi_header gsi_header;
+  pdb_uint expected_size;
+
+  pdb_uint i;
+  pdb_int load_index;
+  pdb_int num_hrs;
+
+  pdb_uint *offsets;
+  pdb_stream offsets_stream;
+
+  pdb_uint hr_bf_max;
+  void *hr_bf;
+
+  pdb_gsi_hr **table = 0;
+
+  if (!pdb_stream_read_struct(gsi_stream, &gsi_header, pdb_gsi_header)) return 0;
+  if (gsi_header.sig != PDB_GSI_SIG) return 0;
+  if (gsi_header.ver != PDB_GSI_V70) return 0;
+  if (gsi_header.hr_len % sizeof(pdb_gsi_file_hr)) return 0;
+  if (gsi_header.num_buckets == 0) return 0;
+
+  hr_bf_max = sizeof(pdb_gsi_hr)*(gsi_header.hr_len/sizeof(pdb_gsi_file_hr));
+  expected_size = hr_bf_max + 
+          (PDB_GSI_HASH_TABLE_MAX + 1)*sizeof(pdb_uint) +
+          sizeof(pdb_gsi_hr *)*(PDB_GSI_HASH_TABLE_MAX + 1);
+
+  /* Reading in all hash records for the hash table.
+   *
+   * Remember that the records that are stored here differ from the ones
+   * that are in the memory. */
+  hr_bf = syms_arena_push(arena, hr_bf_max);
+  if (pdb_stream_read(gsi_stream, hr_bf, gsi_header.hr_len) != gsi_header.hr_len) {
+    SYMS_ASSERT_CORRUPTED_STREAM;
+    return 0;
+  }
+
+  offsets = syms_arena_push_array(arena, pdb_uint, PDB_GSI_HASH_TABLE_MAX + 1);
+  offsets_stream = *gsi_stream;
+  if (!pdb_stream_skip(&offsets_stream, PDB_GSI_BITMAP_SIZE)) {
+    SYMS_ASSERT_CORRUPTED_STREAM;
+    return 0;
+  }
+
+  /* Iterate the bitmap that follows hash records.
+   * Each set bit indicates a valid bucket in the hash table. */
+  for (i = 0; i < PDB_GSI_BITMAP_U32_COUNT; ++i) {
+    pdb_uint bits;
+    pdb_uint k;
+
+    if (!pdb_stream_read_u32(gsi_stream, &bits)) {
+      return 0;
+    }
+
+    for (k = 0; k < 32; ++k) {
+      if (bits & (1 << k)) {
+        pdb_stream_read_u32(&offsets_stream, &offsets[i*32 + k]);
+      } else {
+        offsets[i*32 + k] = ~0u;
+      }
+    }
+  }
+
+  /* Miracles begin here... */
+
+  table   = syms_arena_push_array(arena,  pdb_gsi_hr *, PDB_GSI_HASH_TABLE_MAX);
+  num_hrs = (S32)(gsi_header.hr_len / sizeof(pdb_gsi_file_hr)) - 1;
+
+  for (load_index = PDB_GSI_HASH_TABLE_MAX; load_index >= 0; --load_index) {
+    pdb_gsi_hr *mem_hr;
+    pdb_gsi_file_hr *file_hr;
+    S32 n;
+
+    table[load_index] = 0;
+    if (offsets[load_index] == ~0u) {
+      continue;
+    }
+
+    mem_hr = (pdb_gsi_hr *)hr_bf + num_hrs;
+    file_hr = (pdb_gsi_file_hr *)hr_bf + num_hrs;
+    n = (S32)(offsets[load_index] / 12);
+
+    while (num_hrs >= n) {
+      mem_hr->off = file_hr->off - 1;
+      mem_hr->next = table[load_index];
+      table[load_index] = mem_hr;
+
+      --mem_hr;
+      --file_hr;
+      --num_hrs;
+    }
+  }
+
+  if (num_hrs != -1) {
+    SYMS_ASSERT_CORRUPTED_STREAM;
+    return 0;
+  }
+
+  *out_table = table;
+  *out_num_table = PDB_GSI_HASH_TABLE_MAX;
+
+  return expected_size;
+
+#undef PDB_GSI_HASH_TABLE_MAX  
+#undef PDB_GSI_BITMAP_SIZE     
+#undef PDB_GSI_BITMAP_U32_COUNT
+}
+
+PDB_API syms_bool
+pdb_init_strtable(pdb_strtable *st, pdb_stream *stream)
+{
+  pdb_uint magic, version;
+  syms_bool is_inited = syms_false;
+
+  pdb_stream_read_u32(stream, &magic);
+  pdb_stream_read_u32(stream, &version);
+  if (PDB_STREAM_READ_OR_SEEK_FAILED(stream->flags)) {
+    return syms_false;
+  }
+  if (magic == 0xEFFEEFFEu && version == 1) {
+    pdb_uint strblock_off, strblock_size;
+    pdb_uint bucket_off, bucket_count;
+
+    pdb_stream_read_u32(stream, &strblock_size);
+    strblock_off = stream->off;
+    pdb_stream_skip(stream, strblock_size);
+
+    pdb_stream_read_u32(stream, &bucket_count);
+    bucket_off = stream->off;
+    pdb_stream_skip(stream, bucket_count * sizeof(pdb_ni));
+
+    if (PDB_STREAM_READ_OR_SEEK_FAILED(stream->flags)) {
+      return syms_false;
+    }
+
+    st->magic = magic;
+    st->version = version;
+    st->strblock = pdb_stream_subset(stream, strblock_off, strblock_size);
+    st->buckets = pdb_stream_subset(stream, bucket_off, bucket_count*sizeof(pdb_ni));
+    st->bucket_count = bucket_count;
+
+    is_inited = !pdb_stream_is_null(&st->strblock) && !pdb_stream_is_null(&st->buckets);
+  }
+  return is_inited;
+}
+
+PDB_API syms_bool
+pdb_fileinfo_init(pdb_context *pdb, pdb_fileinfo *fi)
+{
+  syms_bool is_inited = syms_false;
+
+  pdb_stream dbi_stream;
+  if (pdb_stream_init(pdb, PDB_DEFAULT_STREAM_DBI, &dbi_stream)) {
+    pdb_stream fi_stream = pdb_stream_subset(&dbi_stream, pdb->dbi.fileinfo_off, pdb->dbi.fileinfo_len);
+    if (!pdb_stream_is_null(&fi_stream)) {
+      pdb_uint16 mod_count, src_count;
+
+      pdb_stream_read_u16(&fi_stream, &mod_count);
+      pdb_stream_read_u16(&fi_stream, &src_count);
+      if (!PDB_STREAM_READ_OR_SEEK_FAILED(fi_stream.flags)) {
+        pdb_uint ich_size = 0;
+
+
+        fi->mod_count = mod_count;
+        fi->src_count = src_count;
+        // pdb_uint16[mod_count]
+        fi->imod_block = pdb_stream_subset(&fi_stream, fi_stream.off, sizeof(pdb_imod)*fi->mod_count);
+        pdb_stream_skip(&fi_stream, fi->imod_block.size);
+        // pdb_uint16[mod_count]
+        fi->count_block = pdb_stream_subset(&fi_stream, fi_stream.off, sizeof(pdb_uint)*fi->mod_count);
+        {
+          pdb_uint i;
+          for (i = 0; i < mod_count; ++i) {
+            u16 str_count = 0;
+            if (!pdb_stream_read_u16(&fi_stream, &str_count)) {
+              return syms_false;
+            }
+            ich_size += sizeof(pdb_uint)*str_count;
+          }
+        }
+        // pdb_uint[mod_count][count_block[imod]
+        fi->ich_block = pdb_stream_subset(&fi_stream, fi_stream.off, ich_size);
+        pdb_stream_skip(&fi_stream, ich_size);
+        // char*[count_block[imod]]
+        fi->str_block = pdb_stream_subset(&fi_stream, fi_stream.off, fi_stream.size - fi_stream.off);
+        pdb_stream_skip(&fi_stream, fi->str_block.size);
+
+        is_inited = !(PDB_STREAM_READ_OR_SEEK_FAILED(fi_stream.flags)) && 
+                    fi->src_count > 0 && 
+                    fi->mod_count > 0;
+      }
+    }
+  }
+
+  return is_inited;
+}
+
+PDB_API syms_bool
+pdb_fileinfo_get_src_count(pdb_fileinfo *fi, pdb_imod imod, pdb_uint *count_out)
+{
+  if (pdb_stream_seek(&fi->count_block, imod * sizeof(pdb_uint))) {
+    return pdb_stream_read_uint(&fi->count_block, count_out);
+  }
+  return syms_false;
+}
+
+PDB_API syms_bool
+pdb_fileinfo_get_strblock(pdb_fileinfo *fi, pdb_imod imod, pdb_uint *count_out, pdb_stream *strblock_out)
+{
+  pdb_imod i;
+  u16 count;
+  pdb_uint ich_off = 0;
+  pdb_uint ich = PDB_UINT_MAX;
+
+  pdb_stream_seek(&fi->count_block, 0);
+  for (i = 0; i < imod; ++i) {
+    u16 read_count;
+    pdb_stream_read_u16(&fi->count_block, &read_count);
+    ich_off += read_count * sizeof(pdb_uint);
+  }
+  pdb_stream_read_u16(&fi->count_block, &count);
+  pdb_stream_seek(&fi->ich_block, ich_off);
+  pdb_stream_read_uint(&fi->ich_block, &ich);
+  pdb_stream_seek(&fi->str_block, ich);
+
+  *count_out = count;
+  *strblock_out = fi->str_block;
+
+  return !PDB_STREAM_READ_OR_SEEK_FAILED(fi->str_block.flags);
+}
+
+PDB_API syms_bool
+pdb_init(pdb_context *pdb, void *data, pdb_uint data_size)
+{
+  static char PDB70_MAGIC[] = { "Microsoft C/C++ MSF 7.00\r\n\032DS\0\0" };
+  static char PDB20_MAGIC[] = { "Microsoft C/C++ program database 2.00\r\n\x1a\x4a\x47" };
+
+  if (data == NULL || !data_size) {
+    return syms_false;
+  }
+  if (data_size < sizeof(PDB20_MAGIC)) {
+    return syms_false;
+  }
+
+  if (syms_memcmp(data, PDB20_MAGIC, sizeof(PDB20_MAGIC)) == 0) {
+    pdb_header20 *header      = (pdb_header20 *)data;
+    pdb->page_size            = header->page_size;
+    pdb->free_page_map        = (pdb_uint)header->free_page_map;
+    pdb->pages_used           = (pdb_uint)header->pages_used;
+    pdb->root_size            = header->root_size;
+    pdb->page_map_addr.addr16 = &header->page_map_addr;
+    pdb->page_index_size      = sizeof(U16);
+  } else if (syms_memcmp(data, PDB70_MAGIC, sizeof(PDB70_MAGIC)) == 0) {
+    pdb_header70 *header      = (pdb_header70 *)data;
+    pdb->page_size            = header->page_size;
+    pdb->free_page_map        = header->free_page_map >= 0 ? (pdb_uint)header->free_page_map : 0;
+    pdb->pages_used           = header->pages_used >= 0 ? (pdb_uint)header->pages_used : 0;
+    pdb->root_size            = header->root_size;
+    pdb->page_map_addr.addr32 = &header->page_map_addr;
+    pdb->page_index_size      = sizeof(pdb_uint);
+  } else {
+    /* data is not of pdb format */
+    return syms_false;
+  }
+  if (pdb->page_size == 0 || pdb->root_size < sizeof(pdb_uint)) {
+    return syms_false;
+  }
+
+  pdb->file_data = data;
+  pdb->file_size = data_size;
+
+  pdb->publics_array_num = 0;
+  pdb->publics_array = 0;
+
+  pdb->globals_array_num = 0;
+  pdb->globals_array = 0;
+
+  pdb->trampoline_count = 0;
+  pdb->trampoline_contigous = syms_false;
+  pdb->trampoline_data = pdb_cvdata_token_bake(0, 0);
+
+  syms_memset(&pdb->tpi, 0, sizeof(pdb->tpi));
+  syms_memset(&pdb->ipi, 0, sizeof(pdb->ipi));
+
+  {
+    /* NOTE: Info stream is basically a serialized hash table
+     * that contains stream names and corresponding stream number. */
+
+    pdb_stream info_stream;
+    pdb_uint names_len;
+    pdb_uint hash_table_count, hash_table_max;
+    pdb_uint num_present_words, num_deleted_words;
+    pdb_uint names_base_off;
+    pdb_uint epilogue_base_off;
+    pdb_uint i;
+
+    if (!pdb_stream_init(pdb, 1, &info_stream)) {
+      return syms_false;
+    }
+
+    pdb->ver = 0;
+    pdb_stream_read_u32(&info_stream, &pdb->ver);
+    switch (pdb->ver) {
+    case PDB_VER_VC50:
+    case PDB_VER_VC4:
+    case PDB_VER_VC2:
+    case PDB_VER_VC98: {
+      pdb_stream_read_u32(&info_stream, &pdb->time);
+      pdb_stream_read_u32(&info_stream, &pdb->age);
+      pdb_stream_read_u32(&info_stream, &names_len);
+      if (PDB_STREAM_READ_OR_SEEK_FAILED(info_stream.flags)) {
+        return syms_false;
+      }
+      names_base_off = info_stream.off;
+      pdb_stream_skip(&info_stream, names_len);
+      pdb_stream_read_u32(&info_stream, &hash_table_count);
+      pdb_stream_read_u32(&info_stream, &hash_table_max);
+      pdb_stream_read_u32(&info_stream, &num_present_words);
+      pdb_stream_skip(&info_stream, num_present_words*sizeof(pdb_uint));
+      pdb_stream_read_u32(&info_stream, &num_deleted_words);
+      pdb_stream_skip(&info_stream, num_deleted_words*sizeof(pdb_uint));
+      if (PDB_STREAM_READ_OR_SEEK_FAILED(info_stream.flags)) {
+        return syms_false;
+      }
+    } break;
+    case PDB_VER_VC70:
+    case PDB_VER_VC70_DEP:
+    case PDB_VER_VC80:
+    case PDB_VER_VC140:
+    case PDB_VER_VC110: {
+      pdb_stream_read_u32(&info_stream, &pdb->time);
+      pdb_stream_read_u32(&info_stream, &pdb->age);
+      pdb_stream_read(&info_stream, &pdb->auth_guid, sizeof(SymsGUID));
+      pdb_stream_read_u32(&info_stream, &names_len);
+      if (PDB_STREAM_READ_OR_SEEK_FAILED(info_stream.flags)) {
+        return syms_false;
+      }
+      names_base_off = info_stream.off;
+      pdb_stream_skip(&info_stream, names_len);
+      pdb_stream_read_u32(&info_stream, &hash_table_count);
+      pdb_stream_read_u32(&info_stream, &hash_table_max);
+      pdb_stream_read_u32(&info_stream, &num_present_words);
+      pdb_stream_skip(&info_stream, num_present_words*sizeof(pdb_uint));
+      pdb_stream_read_u32(&info_stream, &num_deleted_words);
+      pdb_stream_skip(&info_stream, num_deleted_words*sizeof(pdb_uint));
+      if (PDB_STREAM_READ_OR_SEEK_FAILED(info_stream.flags)) {
+        return syms_false;
+      }
+    } break;
+
+    default: return syms_false;
+    }
+
+#if 0
+    pdb_stream_read_u32(&info_stream, &pdb->time);
+    pdb_stream_read_u32(&info_stream, &pdb->age);
+    pdb_stream_read(&info_stream, &pdb->auth_guid, sizeof(pdb->auth_guid));
+    pdb_stream_read_u32(&info_stream, &names_len);
+    names_base_off = info_stream.off;
+    pdb_stream_skip(&info_stream, names_len);
+
+    pdb_stream_read_u32(&info_stream, &hash_table_count);
+    pdb_stream_read_u32(&info_stream, &hash_table_max);
+    pdb_stream_read_u32(&info_stream, &num_present_words);
+    pdb_stream_skip(&info_stream, num_present_words*sizeof(pdb_uint));
+
+    pdb_stream_read_u32(&info_stream, &num_deleted_words);
+    pdb_stream_skip(&info_stream, num_deleted_words*sizeof(pdb_uint));
+
+    if (PDB_STREAM_READ_OR_SEEK_FAILED(info_stream.flags)) {
+      return syms_false;
+    }
+#endif
+
+    epilogue_base_off = info_stream.off;
+
+    /* NOTE(nick): typed stream is named so because it has a name.
+     * there are three streams that I know of: 
+     *
+     *  "/src/headerblock" 
+     *  "/names"
+     *  "/LinkInfo"
+     *
+     */
+    for (i = 0; i < SYMS_ARRAY_SIZE(pdb->typed_streams); ++i) {
+      pdb->typed_streams[i] = PDB_INVALID_SN;
+    }
+
+    for (i = 0; i < hash_table_count; ++i) {
+      pdb_uint name_off, sn;
+
+      if (!pdb_stream_seek(&info_stream, epilogue_base_off + i*sizeof(pdb_uint)*2)) {
+        break;
+      }
+      if (!pdb_stream_read_u32(&info_stream, &name_off)) {
+        break;
+      }
+      if (!pdb_stream_read_u32(&info_stream, &sn)) {
+        break;
+      }
+      if (!pdb_stream_seek(&info_stream, names_base_off + name_off)) {
+        break;
+      }
+
+      if (pdb_strcmp_stream(syms_string_init_lit("/src/headerblock"), &info_stream)) {
+        pdb->typed_streams[PDB_TYPED_STREAM_HEADER_BLOCK] = (pdb_sn)syms_trunc_u16(sn);
+      } else if (pdb_strcmp_stream(syms_string_init_lit("/names"), &info_stream)) {
+        pdb->typed_streams[PDB_TYPED_STREAM_STRTABLE] = (pdb_sn)syms_trunc_u16(sn);
+      } else if (pdb_strcmp_stream(syms_string_init_lit("/LinkInfo"), &info_stream)) {
+        pdb->typed_streams[PDB_TYPED_STREAM_LINK_INFO] = (pdb_sn)syms_trunc_u16(sn);
+      }
+    }
+
+#if 0 /* I will keep this here, just in case if we need to parse it later  */
+    pdb_uint I;
+    pdb_uint bits[4];
+    c_memset(bits, 0, sizeof(bits));
+
+    for (I = 0; I != num_words; ++I) {
+      pdb_uint word = pdb_stream_read_uint32(&stream);
+      pdb_uint idx;
+
+      for (idx = 0; idx < 32; ++idx) {
+        if (word & (1u << idx)) {
+          pdb_uint idx2 = (I*32) + idx;
+          bits[idx2/32] |= 1l << (idx2 % 32);
+        }
+      }
+    }
+#endif
+  }
+
+  /*
+   * String Table
+   */
+  {
+
+    pdb_stream strtable;
+    if (pdb_stream_init(pdb, pdb->typed_streams[PDB_TYPED_STREAM_STRTABLE], &strtable)) {
+      pdb_uint sig = 0;
+      pdb_uint ver = 0;
+
+      pdb_stream_read_u32(&strtable, &sig);
+      pdb_stream_read_u32(&strtable, &ver);
+
+      if (sig == 0xEFFEEFFEu) {
+        if (ver == 1) {
+          /* 
+           *  Laouyt:
+           *  pdb_uint sig;
+           *  pdb_uint ver;
+           *  pdb_uint strings_len;
+           *  char [][num_strings];
+           *  pdb_uint num_buckets;
+           *  pdb_uint offsets[num_buckets];
+           *  pdb_uint num_strings; 
+           */
+
+          pdb_uint offsets_num, offsets_base;
+          pdb_uint strings_num, strings_base, strings_len;
+
+          if (!pdb_stream_read_u32(&strtable, &strings_len)) {
+            return syms_false;
+          }
+
+          strings_base = strtable.off;
+          if (!pdb_stream_skip(&strtable, strings_len)) {
+            return syms_false;
+          }
+
+          if (!pdb_stream_read_u32(&strtable, &offsets_num)) {
+            return syms_false;
+          }
+
+          offsets_base = strtable.off;
+          pdb_stream_skip(&strtable, offsets_num*sizeof(pdb_stroff));
+
+          /* For some reason the string count is in the end of the stream?! */
+          if (!pdb_stream_read_u32(&strtable, &strings_num)) {
+            return syms_false;
+          }
+
+          /* NOTE(nick; Oct 10 2019): Make sure we read entire table */
+          SYMS_ASSERT(strtable.off == strtable.size);
+
+          pdb->strtable = pdb_stream_subset(&strtable, strings_base, strings_len);
+          pdb->stroffs = pdb_stream_subset(&strtable, offsets_base, offsets_num*sizeof(pdb_stroff));
+        } else {
+          SYMS_ASSERT_FAILURE("unknown version of string table");
+        }
+      }
+    }
+  }
+
+  /*
+   * DBI
+   */
+  {
+    pdb_dbi_header header;
+    pdb_stream dbi_stream;
+    pdb_uint streams_off;
+    syms_bool inited = syms_false;
+
+    pdb_mod_it mod_it;
+    pdb_mod mod;
+
+    pdb_sec_it sec_it;
+    pdb_img_sec sec;
+
+    syms_memset(&header, 0, sizeof(header));
+
+    if (pdb_stream_init(pdb, PDB_DEFAULT_STREAM_DBI, &dbi_stream)) {
+      if (pdb_stream_read_struct(&dbi_stream, &header, pdb_dbi_header)) {
+        if (header.sig == PDB_DBI_HEADER_SIG_V1) {
+          if (header.module_info_size % sizeof(pdb_uint) == 0) {
+            if (header.sec_con_size % sizeof(pdb_uint) == 0) {
+              if (header.sec_map_size % sizeof(pdb_uint) == 0) {
+                if (header.file_info_size % sizeof(pdb_uint) == 0) {
+                  inited = syms_true;
+                }
+              }
+            }
+          }
+        }
+      }
+    }
+    if (!inited) return syms_false;
+
+    {
+      pdb_uint i;
+      for (i = 0; i < SYMS_ARRAY_SIZE(pdb->dbi.dbg_streams); ++i) {
+        pdb->dbi.dbg_streams[i] = PDB_INVALID_SN;
+      }
+    }
+    streams_off = sizeof(pdb_dbi_header)
+            + header.module_info_size
+            + header.sec_con_size
+            + header.sec_map_size
+            + header.file_info_size
+            + header.tsm_size
+            + header.ec_info_size; 
+    if (pdb_stream_seek(&dbi_stream, streams_off)) {
+      pdb_stream_read(&dbi_stream, &pdb->dbi.dbg_streams[0], sizeof(pdb->dbi.dbg_streams));
+    }
+
+    pdb->dbi.header = header;
+
+    pdb->dbi.ver          = header.version;
+    pdb->dbi.machine_type = header.machine;
+    pdb->dbi.symrec_sn    = header.sym_record_sn;
+    pdb->dbi.pubsym_sn    = header.public_sym_sn;
+    pdb->dbi.globalsym_sn = header.global_sym_sn;
+
+    pdb->dbi.modinfo_off = sizeof(pdb_dbi_header);
+    pdb->dbi.modinfo_len = header.module_info_size;
+
+    pdb->dbi.seccon_off = sizeof(pdb_dbi_header) + header.module_info_size;
+    pdb->dbi.seccon_len = header.sec_con_size;
+
+    /* NOTE: If you want to parse the section map sub stream from the DBI
+     * here is the layout:
+     * U16 sec_count;
+     * U16 logical_sec_count;
+     *
+     * After follows an array the count for which is defined by the
+     * sec_count field above.
+     *  U16 flags;
+     *      Read              = 1 << 0,
+     *      Write             = 1 << 1,
+     *      Execute           = 1 << 2,
+     *      AddressIS32bit= 1 << 3,
+     *      IsSelector        = 1 << 8,
+     *      IsAbsoluteAddress = 1 << 9,
+     *      IsGroup           = 1 << 10
+     *  U16 ovl; Logical overlay number
+     *  U16 group; Group index into descriptor array
+     *  U16 frame;
+     *  U16 sec_name; Byte index of the section or group name
+     *                in the sstSegName table, or 0xFFFF
+     *  U16 class_name; Same thing as sec_name
+     *  pdb_uint offset; Byte offset of the logical section
+     *              within the specified physical section.
+     *              If group flag is set, offset is the
+     *              offset of the group.
+     *  pdb_uint sec_byte_length; Byte count of section or group. */
+    pdb->dbi.secmap_off = sizeof(pdb_dbi_header) + header.module_info_size + header.sec_con_size;
+    pdb->dbi.secmap_len = header.sec_map_size;
+
+    pdb->dbi.fileinfo_off = sizeof(pdb_dbi_header) + header.module_info_size + header.sec_con_size + header.sec_map_size;
+    pdb->dbi.fileinfo_len = header.file_info_size;
+
+    /* NOTE: If you want to parse the EC info sub stream from the DBI
+     * here is the layout:
+     * pdb_uint sig; 0xeffeeffe
+     * pdb_uint ver; This filed indicates which hashing function was used
+     * pdb_uint len; Length of the buffer that follows
+     * char names[][len]; Strings containing file names
+     *
+     * For more info look in microsoft's github for file nmt.h */
+    pdb->dbi.ecinfo_off = sizeof(pdb_dbi_header) + header.module_info_size + header.sec_con_size + header.sec_map_size +
+                header.file_info_size + header.tsm_size;
+    pdb->dbi.ecinfo_len = header.ec_info_size;
+
+    pdb->dbi.secs_num = 0;
+    pdb->dbi.secs     = 0;
+
+    if (pdb_sec_it_init(pdb, &sec_it)) {
+      while (pdb_sec_it_next(&sec_it, &sec)) {
+        ++pdb->dbi.secs_num;
+      }
+    }
+#if defined(SYMS_PARANOID)
+    SYMS_ASSERT(pdb->dbi.secs_num);
+#endif
+
+    pdb->dbi.mods_num = 0;
+    pdb->dbi.mods = 0;
+    if (pdb_mod_it_init(&mod_it, pdb)) {
+      while (pdb_mod_it_next(&mod_it, &mod)) {
+        ++pdb->dbi.mods_num;
+      }
+    }
+#if defined(SYMS_PARANOID)
+    SYMS_ASSERT(pdb->dbi.mods_num);
+#endif
+  }
+
+  /*
+   * Types
+   */
+  {
+    static syms_bool init_name_table = syms_true;
+    static SymsString name_table[PDB_BASIC_TYPE_MAX];
+
+    pdb_stream tm_data;
+
+    if (init_name_table) {
+#define X(name, display_name, value) name_table[value] = syms_string_init_lit(display_name);
+      PDB_BASIC_TYPE_LIST
+#undef X
+      init_name_table = syms_false;
+    }
+
+    pdb->basic_typenames_array_num = SYMS_ARRAY_SIZE(name_table);
+    pdb->basic_typenames_array     = name_table;
+
+    if (pdb_stream_init(pdb, PDB_DEFAULT_STREAM_TPI, &tm_data)) {
+      if (pdb_stream_read_struct(&tm_data, &pdb->tpi.header, pdb_tm_header)) {
+        if (pdb->tpi.header.ti_hi > pdb->tpi.header.ti_lo) {
+          pdb->tpi.sn = PDB_DEFAULT_STREAM_TPI;
+        }
+      }
+    }
+
+    if (pdb_stream_init(pdb, PDB_DEFAULT_STREAM_IPI, &tm_data)) {
+      if (pdb_stream_read_struct(&tm_data, &pdb->ipi.header, pdb_tm_header)) {
+        if (pdb->ipi.header.ti_hi > pdb->ipi.header.ti_lo) {
+          pdb->ipi.sn = PDB_DEFAULT_STREAM_IPI;
+        }
+      }
+    }
+  }
+
+  // pdb_dump_lines_c11(pdb);
+
+  return syms_true;
+}
+
+PDB_API SymsUMM
+pdb_load_types(pdb_context *pdb, struct SymsArena *arena)
+{
+  SymsUMM arena_size_begin;
+  SymsUMM arena_used;
+
+  if (!arena) {
+    return pdb_calc_size_for_types(&pdb->tpi.header) + pdb_calc_size_for_types(&pdb->ipi.header);
+  }
+
+  arena_used = 0;
+  arena_size_begin = arena->size;
+  if (pdb_tm_init(&pdb->tpi, pdb, PDB_DEFAULT_STREAM_TPI, arena)) {
+    arena_used += (arena->size - arena_size_begin);
+  }
+
+  arena_size_begin = arena->size;
+  if (pdb_tm_init(&pdb->ipi, pdb, PDB_DEFAULT_STREAM_IPI, arena)) {
+    arena_used += (U32)(arena->size - arena_size_begin);
+  }
+
+  return arena_used;
+}
+
+PDB_API SymsUMM
+pdb_load_publics(pdb_context *pdb, struct SymsArena *arena)
+{
+  SymsUMM result = 0;
+  pdb_stream stream;
+
+  if (pdb_stream_init(pdb, pdb->dbi.pubsym_sn, &stream)) {
+    pdb_psi_header header;
+    if (pdb_stream_read_struct(&stream, &header, pdb_psi_header)) {
+      result = pdb_parse_gsi_hash_table(&stream, arena, &pdb->publics_array, &pdb->publics_array_num);
+    }
+  }
+
+  return result;
+}
+
+PDB_API SymsUMM
+pdb_load_globals(pdb_context *pdb, struct SymsArena *arena)
+{
+  SymsUMM result = 0;
+  pdb_stream gsi_stream;
+
+  if (pdb_stream_init(pdb, pdb->dbi.globalsym_sn, &gsi_stream)) {
+    result = pdb_parse_gsi_hash_table(
+                    &gsi_stream, 
+                    arena, 
+                    &pdb->globals_array, 
+                    &pdb->globals_array_num);
+  }
+
+  return result;
+}
+
+PDB_API SymsUMM
+pdb_load_trampolines(pdb_context *pdb, struct SymsArena *arena)
+{
+  if (!arena) {
+    pdb_mod_it mod_it;
+
+    pdb->trampoline_count = 0;
+    pdb->trampoline_contigous = syms_true;
+    pdb->trampoline_data = pdb_cvdata_token_bake(0, 0);
+
+    if (pdb_mod_it_init(&mod_it, pdb)) {
+      pdb_mod mod;
+      while (pdb_mod_it_next(&mod_it, &mod)) {
+        if (pdb_strcmp_pointer_nocase(pdb, syms_string_init_lit("* Linker *"), &mod.name)) {
+          pdb_stream syms_stream;
+          if (pdb_mod_get_debug_sec(&mod, PDB_MOD_SEC_SYMS, &syms_stream)) {
+            pdb_cvdata_token token;
+            pdb_sym_it sym_it;
+
+            token.sn = syms_stream.sn;
+            token.soffset = syms_stream.off;
+            if (pdb_sym_it_init_token(&sym_it, pdb, token)) {
+              pdb_cv_sym_type cv_type;
+
+              pdb_uint sig = 0;
+              pdb_stream_read_u32(&sym_it.stream, &sig);
+
+              for (;;) {
+                pdb_uint off = sym_it.stream.off;
+                if (pdb_sym_it_next(&sym_it, &cv_type, 0)) {
+                  if (cv_type == PDB_CV_SYM_TRAMPOLINE) {
+                    pdb->trampoline_data.sn = sym_it.stream.sn;
+                    pdb->trampoline_data.soffset = off;
+                    break;
+                  }
+                } else {
+                  break;
+                }
+              }
+              while (pdb_sym_it_next(&sym_it, &cv_type, 0)) {
+                if (cv_type == PDB_CV_SYM_TRAMPOLINE) {
+                  pdb->trampoline_count += 1;
+                } else {
+                  break;
+                }
+              }
+              while (pdb_sym_it_next(&sym_it, &cv_type, 0)) {
+                if (cv_type == PDB_CV_SYM_TRAMPOLINE) {
+                  pdb->trampoline_contigous = syms_false;
+                  pdb->trampoline_count += 1;
+                }
+              }
+            }
+          }
+        }
+      }
+    }
+  }
+
+  return 0;
+}
+
+PDB_API syms_bool
+pdb_trampoline_from_ip(pdb_context *pdb, pdb_isec src_sec, pdb_isec_umm src_sec_off, pdb_isec *sec_out, pdb_isec_umm *off_out)
+{
+  pdb_sym_it sym_it;
+
+  src_sec += 1;
+  if (pdb->trampoline_count > 0) {
+    if (pdb_sym_it_init_token(&sym_it, pdb, pdb->trampoline_data)) {
+      if (pdb->trampoline_contigous) {
+        pdb_uint min = 0;
+        pdb_uint max = pdb->trampoline_count - 1;
+        pdb_uint base_offset = pdb->trampoline_data.soffset;
+        while (min <= max) {
+          pdb_uint mid = (min + max)/2;
+          pdb_cv_trampolinesym t_mid;
+
+          syms_memset(&t_mid, 0, sizeof(t_mid));
+
+          pdb_stream_seek(&sym_it.stream, base_offset + mid*(sizeof(pdb_cv_trampolinesym) + 4) + 4);
+          pdb_stream_read(&sym_it.stream, &t_mid, sizeof(t_mid));
+
+          if (src_sec < t_mid.thunk_sec) {
+            max = mid - 1;
+          } else if (src_sec > t_mid.thunk_sec) {
+            min = mid + 1;
+          } else {
+            if (src_sec_off < t_mid.thunk_sec_off) {
+              max = mid - 1;
+            } else if (src_sec_off >= (t_mid.thunk_sec_off + t_mid.thunk_size)) {
+              min = mid + 1;
+            } else {
+              *sec_out = t_mid.target_sec;
+              *off_out = t_mid.target_sec_off;
+              return syms_true;
+            }
+          }
+        }
+      } else {
+        pdb_cv_sym_type cv_type;
+        pdb_stream cv_data;
+
+        while (pdb_sym_it_next(&sym_it, &cv_type, &cv_data)) {
+          if (cv_type == PDB_CV_SYM_TRAMPOLINE) {
+            pdb_cv_trampolinesym trampoline;
+            if (pdb_stream_read(&cv_data, &trampoline, sizeof(trampoline))) {
+              if (src_sec == trampoline.thunk_sec) {
+                if (src_sec_off >= trampoline.thunk_sec_off && src_sec_off < (trampoline.thunk_sec_off + trampoline.thunk_size)) {
+                  *sec_out = trampoline.target_sec;
+                  *off_out = trampoline.target_sec_off;
+                  return syms_true;
+                }
+              }
+            }
+          }
+        }
+      }
+    }
+  }
+  return syms_false;
+}
+
+PDB_API SymsUMM
+pdb_load_dbi(pdb_context *pdb, struct SymsArena *arena)
+{
+  SymsUMM arena_size;
+  SymsUMM dbi_size;
+
+  syms_bool inited;
+  pdb_stream sec_stream;
+
+  if (!arena) {
+    SymsUMM num_bytes = 32; /* ALIGNMENT */
+    num_bytes += pdb->dbi.mods_num*sizeof(pdb->dbi.mods[0]);
+    num_bytes += pdb->dbi.secs_num*sizeof(pdb->dbi.secs[0]);
+    return num_bytes;
+  }
+
+  arena_size = arena->size;
+
+  inited = syms_false;
+  if (pdb_stream_init(pdb, pdb->dbi.dbg_streams[PDB_DBG_STREAM_SECTION_HEADER], &sec_stream)) {
+    pdb->dbi.mods = syms_arena_push_array(arena, pdb_uint, pdb->dbi.mods_num);
+    pdb->dbi.secs = syms_arena_push_array(arena, pdb_img_sec, pdb->dbi.secs_num);
+    inited = pdb->dbi.mods && pdb->dbi.secs;
+  }
+
+  if (inited) {
+    pdb_uint index;
+
+    pdb_mod_it mod_it;
+    pdb_sec_it sec_it;
+
+    index = 0;
+    if (pdb_mod_it_init(&mod_it, pdb)) {
+      for (;;) {
+        pdb_mod mod;
+        pdb->dbi.mods[index++] = mod_it.dbi_data.off;
+        if (!pdb_mod_it_next(&mod_it, &mod)) break;
+        if (index >= pdb->dbi.mods_num)      break;
+      }
+    }
+    inited = index >= pdb->dbi.mods_num;
+
+    if (inited) {
+      index = 0;
+      if (pdb_sec_it_init(pdb, &sec_it)) {
+        while (pdb_sec_it_next(&sec_it, &pdb->dbi.secs[index++])) {
+          if (index >= pdb->dbi.secs_num) {
+            break;
+          }
+        }
+      }
+      inited = index >= pdb->dbi.secs_num;
+    }
+  }
+
+  dbi_size = arena->size - arena_size;
+  if (!inited) {
+    dbi_size = 0;
+  }
+
+  return syms_trunc_u32(dbi_size);
+}
+
+PDB_API SymsUMM
+pdb_calc_size_for_aux_data(pdb_context *pdb)
+{
+  SymsUMM result = 0;
+
+  result += pdb_load_types(pdb, 0);
+  result += pdb_load_publics(pdb, 0);
+  result += pdb_load_globals(pdb, 0);
+  result += pdb_load_dbi(pdb, 0);
+  result += pdb_load_trampolines(pdb, 0);
+  result += 128; /* For alignment */
+
+  return result;
+}
+
+PDB_API syms_bool
+pdb_load_aux_data(pdb_context *pdb, SymsArena *arena)
+{
+  pdb_load_dbi(pdb, arena);
+  pdb_load_types(pdb, arena);
+  pdb_load_publics(pdb, arena);
+  pdb_load_globals(pdb, arena);
+  pdb_load_trampolines(pdb, arena);
+  return syms_true;
+}
+
+PDB_API SymsNTFileHeaderMachineType
+pdb_get_machine_type(pdb_context *pdb)
+{
+  return pdb->dbi.machine_type;
+}
+
+#if 0
+PDB_API void
+pdb_dump_msf_small(pdb_context *pdb)
+{
+  SYMS_ASSERT(pdb->page_index_size == sizeof(U16));
+  U32 off = 0;
+  U32 stream_count;
+  if (!pdb_root_read(pdb, off, sizeof(stream_count), &stream_count)) {
+    return;
+  }
+  U32 *stream_sizes = (U32 *)malloc(sizeof(U32)*stream_count);
+  off += 4;
+  printf("Stream Count: %u\n", stream_count);
+  {
+    U32 i;
+    for (i = 0; i < stream_count; ++i) {
+      U32 size;
+      U32 unknown;
+      if (!pdb_root_read(pdb, off, sizeof(size), &size)) {
+        return;
+      }
+      off += 4;
+      if (!pdb_root_read(pdb, off, sizeof(unknown), &unknown)) {
+        return;
+      }
+      off += 4;
+      printf("SN %u: size %u, unknown: 0x%x\n", i, size, unknown);
+      stream_sizes[i] = size;
+    }
+    for (i = 0; i < stream_count; ++i) {
+      U32 index_count = pdb_count_pages(pdb, stream_sizes[i]);
+      U32 k = 0;
+      printf("SN %u\n", i);
+      while (index_count > 0) {
+        U16 page_index;
+        if (!pdb_root_read(pdb, off, sizeof(page_index), &page_index)) {
+          return;
+        }
+        k += 1;
+        if (k % 4 == 0) {
+          printf("\n");
+        }
+        printf("0x%03X ", page_index);
+        off += sizeof(page_index);
+        index_count -= 1;
+      }
+      printf("\n\n");
+    }
+    SYMS_ASSERT(off == pdb->root_size);
+  }
+}
+
+PDB_API void
+pdb_dump_lines_c11(pdb_context *pdb)
+{
+  pdb_mod_it mod_it;
+  if (pdb_mod_it_init(&mod_it, pdb)) {
+    pdb_mod mod;
+    while (pdb_mod_it_next(&mod_it, &mod)) {
+      pdb_stream c11_lines;
+      if (pdb_mod_get_debug_sec(&mod, PDB_MOD_SEC_LINES_C11, &c11_lines)) {
+        U16 file_count;
+        pdb_stream_read_u16(&c11_lines, &file_count);
+        U16 range_count;
+        pdb_stream_read_u16(&c11_lines, &range_count);
+        /* NOTE(nick): offsets into c11_lines stream where
+         * file info starts. useful for index to file translations */
+        for (U16 i=0;i<file_count;++i) {
+          U32 offset;
+          pdb_stream_read_u32(&c11_lines, &offset);
+        }
+        /* NOTE(nick): range in which instructions were committed */
+        for (U16 i=0; i<range_count; ++i) {
+          U32 start, end;
+          pdb_stream_read_u32(&c11_lines, &start);
+          pdb_stream_read_u32(&c11_lines, &end);
+        }
+        SYMS_ASSERT(!PDB_STREAM_READ_OR_SEEK_FAILED(c11_lines.flags));
+        for (U16 i=0;i<range_count;++i){
+          U16 seg;
+          pdb_stream_read_u16(&c11_lines, &seg);
+        }
+        pdb_stream_align(&c11_lines, 4);
+        SYMS_ASSERT(!PDB_STREAM_READ_OR_SEEK_FAILED(c11_lines.flags));
+        for (U16 file_index=0;file_index<file_count;++file_index) {
+          U16 seg_count, pad;
+          pdb_stream_read_u16(&c11_lines, &seg_count);
+          pdb_stream_read_u16(&c11_lines, &pad);
+          /* NOTE(nick): offsets for line sections */
+          for (U32 j = 0; j < seg_count; ++j) {
+            U32 base_src_line;
+            pdb_stream_read_u32(&c11_lines, &base_src_line);
+          }
+          SYMS_ASSERT(!PDB_STREAM_READ_OR_SEEK_FAILED(c11_lines.flags));
+          /* NOTE(nick): section start offset and section end offset */
+          for (U32 j = 0; j < seg_count; ++j) {
+            U32 start, end;
+            pdb_stream_read_u32(&c11_lines, &start);
+            pdb_stream_read_u32(&c11_lines, &end);
+          }
+          SYMS_ASSERT(!PDB_STREAM_READ_OR_SEEK_FAILED(c11_lines.flags));
+          {
+            U32 size = 0;
+            char buffer[1024];
+            switch (pdb->ver) {
+            case PDB_VER_VC50:
+            case PDB_VER_VC4:
+            case PDB_VER_VC2:
+            case PDB_VER_VC98: {
+              pdb_stream_read_u08(&c11_lines, (U8*)&size);
+              pdb_stream_read(&c11_lines, &buffer[0], size);
+            } break;
+            case PDB_VER_VC70:
+            case PDB_VER_VC70_DEP:
+            case PDB_VER_VC80:
+            case PDB_VER_VC140:
+            case PDB_VER_VC110: {
+              size = pdb_stream_read_str(&c11_lines, &buffer[0], sizeof(buffer));
+              pdb_stream_skip(&c11_lines, size);
+            } break;
+            }
+            printf("%.*s\n", size, buffer);
+            pdb_stream_align(&c11_lines, 4);
+          }
+          for (U32 seg_index = 0; seg_index < seg_count; ++seg_index) {
+            U16 sec_id, pair_count;
+            pdb_stream_read_u16(&c11_lines, &sec_id);
+            pdb_stream_read_u16(&c11_lines, &pair_count);
+            U32 c = 0;
+            printf("\n\t");
+            for (U16 j = 0; j < pair_count; ++j) {
+              U32 offset;
+              pdb_stream_read_u32(&c11_lines, &offset);
+              printf("%X ", offset);
+              if (c % 5 == 4) {
+                printf("\n\t");
+              }
+              c += 1;
+            }
+            c = 0;
+            printf("\n\n\t");
+            for (U16 j = 0; j < pair_count; ++j) {
+              U16 ln;
+              pdb_stream_read_u16(&c11_lines, &ln);
+              printf("%u ", ln);
+              if (c % 5 == 4) {
+                printf("\n\t");
+              }
+              c += 1;
+            }
+            printf("\n");
+            pdb_stream_align(&c11_lines, 4);
+            SYMS_ASSERT(!PDB_STREAM_READ_OR_SEEK_FAILED(c11_lines.flags));
+          }
+          printf("\n");
+        }
+      }
+    }
+  }
+}
+
+PDB_API void
+pdb_dump_bytes(pdb_stream *data, U32 count)
+{
+  U32 i;
+  for (i = 0; i < count; ++i) {
+    U8 b;
+    if (!pdb_stream_read_u08(data, &b)) {
+      break;
+    }
+    printf("%x ", b);
+  }
+  printf("\n");
+}
+
+PDB_API void
+pdb_dump_fastlink(pdb_context *pdb)
+{
+  pdb_mod_it mod_it;
+  if (pdb_mod_it_init(&mod_it, pdb)) {
+    pdb_mod mod;
+    while (pdb_mod_it_next(&mod_it, &mod)) {
+      pdb_sym_it sym_it;
+      
+      {
+        char name[4096];
+        U32 name_size = pdb_pointer_read(pdb, &mod.name, 0, name, sizeof(name));
+        printf("%.*s\n", name_size, name);
+      }
+
+      if (pdb_sym_it_init(&sym_it, &mod)) {
+        pdb_cv_sym_type type;
+        pdb_stream data;
+        while (pdb_sym_it_next(&sym_it, &type, &data)) {
+          if (type == PDB_CV_SYM_FASTLINK) {
+            pdb_cv_fastlink sym;
+            /*pdb_dump_bytes(&data, sizeof(pdb_cv_fastlink));*/
+
+            if (pdb_stream_read_struct(&data, &sym, pdb_cv_fastlink)) {
+              if (sym.flags & PDB_CV_FASTLINK_FLAGS_IS_UDT || sym.flags & PDB_CV_FASTLINK_FALGS_IS_DATA || sym.flags & PDB_CV_FASTLINK_FLAGS_IS_CONST ||
+                  sym.flags & PDB_CV_FASTLINK_FLAGS_IS_NAMESPACE) {
+                continue;
+                printf("TypeIndex 0x%X, ", sym.u.type_index);
+              } else {
+                printf("CoffIsect 0x%X, ", sym.u.unknown);
+              }
+              printf("Flags 0x%X, ", sym.flags);
+
+              U8 temp[4096];
+              pdb_stream_read_str(&data, temp, sizeof(temp));
+              printf("%s\n", temp);
+              printf("\n");
+            }
+          }
+        }
+      }
+    }
+  }
+}
+
+#endif
