RAD symbol library

diff --git a/Engine/Source/Runtime/SymsLib/src/pdb/syms_pdb.h b/Engine/Source/Runtime/SymsLib/src/pdb/syms_pdb.h
new file mode 100644
--- /dev/null
+++ b/Engine/Source/Runtime/SymsLib/src/pdb/syms_pdb.h
@@ -0,0 +1,1445 @@
+// Copyright Epic Games, Inc. All Rights Reserved.
+#ifndef SYMS_PDB_INCLUDE_H
+#define SYMS_PDB_INCLUDE_H
+
+#define SYMS_ASSERT_CORRUPTED_STREAM SYMS_ASSERT_FAILURE("corrupted stream")
+
+/* Maximum number of bytes a string contains */
+#define PDB_STRLEN_MAX 4096
+
+enum
+{
+  PDB_ERROR_ok,
+  PDB_ERROR_corrupted_data
+};
+typedef u32 pdb_error_code;
+
+typedef enum 
+{
+  PDB_VER_VC2      = 0x13048EA,
+  PDB_VER_VC4      = 0x1306C1F,
+  PDB_VER_VC41     = 0x1306CDE,
+  PDB_VER_VC50     = 0x13091F3,
+  PDB_VER_VC98     = 0x130BA2C,
+  PDB_VER_VC70     = 0x1312E94,
+  PDB_VER_VC70_DEP = 0x131084C,
+  PDB_VER_VC80     = 0x131A5B5,
+  PDB_VER_VC110    = 0x1329141,
+  PDB_VER_VC140    = 0x13351DC
+} pdb_ver_e;
+
+typedef enum 
+{
+  PDB_TYPED_STREAM_HEADER_BLOCK,
+
+  /* Contains hash tables for strings and offsets. 
+   * For parsing look at "pdb_init". */
+  PDB_TYPED_STREAM_STRTABLE,
+
+  /* Contains info left there by linker. */
+  PDB_TYPED_STREAM_LINK_INFO,
+
+  PDB_TYPED_STREAM_MAX
+} pdb_typed_stream_e;
+
+typedef enum 
+{
+  PDB_DEFAULT_STREAM_ROOT,
+  PDB_DEFAULT_STREAM_STREAMS,
+  PDB_DEFAULT_STREAM_TPI,
+  PDB_DEFAULT_STREAM_DBI,
+  PDB_DEFAULT_STREAM_IPI
+} pdb_default_stream_e;
+typedef pdb_sn pdb_default_stream;
+
+typedef struct
+{
+  u32 count;
+  pdb_cv_itype *e;
+} pdb_itype_array;
+
+#define PDB_NUMERIC_MAX 32
+typedef struct pdb_numeric 
+{
+  pdb_cv_itype itype;
+  union 
+  {
+    char ch;
+
+    S8 int8;
+    S16 int16;
+    S32 int32;
+    S64 int64;
+
+    U8 uint8;
+    U16 uint16;
+    U32 uint32;
+    U64 uint64;
+
+    F32 float32;
+    F64 float64;
+
+    U8 data[PDB_NUMERIC_MAX];
+  } u;
+} pdb_numeric;
+
+enum 
+{
+  PDB_POINTER_MODE_NULL,
+  PDB_POINTER_MODE_RAW,
+  PDB_POINTER_MODE_PAGES,
+  PDB_POINTER_MODE_STREAM
+};
+typedef u32 pdb_pointer_mode;
+
+// Maximum number of pages we can bake in one pointer.
+// Imperically largerst thing that gets baked is a string,
+// which are around 4k long and by default page size in PDB
+// is 4K, so edge case scenario here is a string that spans 
+// accross two pages.
+#define PDB_POINTER_PAGE_MAX 2
+
+typedef struct pdb_pointer 
+{
+  pdb_pointer_mode mode;
+
+  union {
+    struct {
+      const void *data;
+      pdb_uint    size;
+    } raw;
+
+    struct {
+      pdb_uint offs[PDB_POINTER_PAGE_MAX];
+      pdb_uint size[PDB_POINTER_PAGE_MAX];
+    } pages;
+  
+    struct {
+      pdb_sn   sn;
+      pdb_uint off;
+      pdb_uint size;
+    } stream;
+  } u;
+} pdb_pointer;
+typedef pdb_pointer pdb_string_ref;
+
+PDB_API pdb_pointer
+pdb_pointer_bake(pdb_stream *stream, pdb_uint num_bytes);
+
+PDB_API pdb_pointer
+pdb_pointer_bake_null(void);
+
+PDB_API pdb_pointer
+pdb_pointer_bake_one_page(void *data, U32 data_len);
+
+PDB_API pdb_pointer
+pdb_pointer_bake_str(SymsString str);
+
+PDB_API pdb_pointer
+pdb_pointer_bake_stream_str(pdb_stream *msf);
+
+PDB_API pdb_pointer
+pdb_pointer_bake_sn(struct pdb_context *pdb, pdb_sn sn, U32 off, U32 len);
+
+PDB_API U32
+pdb_pointer_get_size(pdb_pointer *pointer);
+
+PDB_API U32
+pdb_pointer_read(struct pdb_context *pdb, const pdb_pointer *pointer, U32 off, void *bf, U32 bf_max);
+
+PDB_API U32
+pdb_pointer_read_u32(struct pdb_context *pdb, pdb_pointer *pointer, U32 off);
+
+PDB_API U16
+pdb_pointer_read_u16(struct pdb_context *pdb, pdb_pointer *pointer, U32 off);
+
+PDB_API U8
+pdb_pointer_read_u08(struct pdb_context *pdb, pdb_pointer *pointer, U32 off);
+
+PDB_API U32
+pdb_pointer_read_utf8(struct pdb_context *pdb, pdb_pointer *pointer, U32 off, U32 *codepoint_out);
+
+PDB_API syms_bool
+pdb_pointer_cmp(struct pdb_context *pdb, pdb_pointer *pointer_a, pdb_pointer *pointer_b);
+
+typedef enum 
+{
+  PDB_STRCMP_FLAG_NULL     = 0,
+  PDB_STRCMP_FLAG_NOCASE   = (1 << 0),
+  PDB_STRCMP_FLAG_NOREWIND = (1 << 1)
+} pdb_strcmp_flags_e;
+
+PDB_API syms_bool
+pdb_pointer_strcmp_(struct pdb_context *pdb, pdb_pointer *pointer_a, pdb_pointer *pointer, pdb_strcmp_flags_e cmp_flags);
+
+#define pdb_strcmp_pointer(pdb, a, b)        pdb_strcmp_pointer_(pdb, a, b, PDB_STRCMP_FLAG_NULL)
+#define pdb_strcmp_pointer_nocase(pdb, a, b) pdb_strcmp_pointer_(pdb, a, b, PDB_STRCMP_FLAG_NOCASE)
+PDB_API syms_bool
+pdb_strcmp_pointer_(struct pdb_context *pdb, SymsString str, pdb_pointer *str_b, pdb_strcmp_flags_e cmp_flags);
+
+#define pdb_pointer_strcmp(pdb, a, b)        pdb_pointer_strcmp_(pdb, a, b, PDB_STRCMP_FLAG_NULL)
+#define pdb_pointer_strcmp_nocase(pdb, a, b) pdb_pointer_strcmp_(pdb, a, b, PDB_STRCMP_FLAG_NOCASE)
+PDB_API syms_bool
+pdb_pointer_strcmp_(struct pdb_context *pdb, pdb_pointer *str_a, pdb_pointer *str_b, pdb_strcmp_flags_e cmp_flags);
+
+#define pdb_strcmp_stream(a, b)        pdb_strcmp_stream_(a, b, PDB_STRCMP_FLAG_NULL)
+#define pdb_strcmp_stream_nocase(a, b) pdb_strcmp_stream_(a, b, PDB_STRCMP_FLAG_NOCASE)
+PDB_API syms_bool
+db_strcmp_stream_(SymsString str, struct pdb_stream *stream, pdb_strcmp_flags_e cmp_flags);
+
+#define pdb_stream_strcmp(a, b)        pdb_stream_strcmp_stream_(a, b, PDB_STRCMP_FLAG_NULL)
+#define pdb_stream_strcmp_nocase(a, b) pdb_stream_strcmp_stream_(a, b, PDB_STRCMP_FLAG_NOCASE)
+PDB_API syms_bool
+pdb_stream_strcmp_stream_(pdb_stream *stream_a, pdb_stream *stream_b, pdb_strcmp_flags_e cmp_flags);
+
+#define pdb_stream_strcmp_pointer(a, b)         pdb_stream_strcmp_pointer_(a, b, PDB_STRCMP_FLAG_NULL)
+#define pdb_stream_strcmp_pointer_nocase(a, b)  pdb_stream_strcmp_pointer_(a, b, PDB_STRCMP_FLAG_NOCASE)
+PDB_API syms_bool
+pdb_stream_strcmp_pointer_(pdb_stream *str_a, pdb_pointer *str_b, pdb_strcmp_flags_e cmp_flags);
+
+typedef enum 
+{
+  PDB_TYPE_NULL,
+  PDB_TYPE_CHAR,
+  PDB_TYPE_UCHAR,
+  PDB_TYPE_INT8,
+  PDB_TYPE_UINT8,
+  PDB_TYPE_INT16,
+  PDB_TYPE_UINT16,
+  PDB_TYPE_INT32,
+  PDB_TYPE_UINT32,
+  PDB_TYPE_INT64,
+  PDB_TYPE_UINT64,
+  PDB_TYPE_INT128,
+  PDB_TYPE_UINT128,
+  PDB_TYPE_REAL16,
+  PDB_TYPE_REAL32,
+  PDB_TYPE_REAL32PP, /* 32 bit partial precision */
+  PDB_TYPE_REAL64,
+  PDB_TYPE_REAL80,
+  PDB_TYPE_REAL128,
+  PDB_TYPE_BOOL,
+  PDB_TYPE_VOID,
+  PDB_TYPE_WCHAR,
+  PDB_TYPE_STRUCT,
+  PDB_TYPE_UNION,
+  PDB_TYPE_CLASS,
+  PDB_TYPE_ENUM,
+  PDB_TYPE_PROC,
+  PDB_TYPE_ENUM_FIELD,
+  PDB_TYPE_BITFIELD,
+  PDB_TYPE_NESTED_UDT,
+  PDB_TYPE_FWDREF,
+  PDB_TYPE_VTSHAPE,
+  PDB_TYPE_METHOD,
+  PDB_TYPE_METHODLIST,
+  PDB_TYPE_ARGLIST,
+  PDB_TYPE_VARIADIC,
+  PDB_TYPE_POINTER,
+  PDB_TYPE_PTR,
+  PDB_TYPE_ARR,
+  PDB_TYPE_FUNCID,
+  PDB_TYPE_MFUNCID,
+  PDB_TYPE_STRINGID,
+  PDB_TYPE_COMPLEX32,
+  PDB_TYPE_COMPLEX64,
+  PDB_TYPE_COMPLEX80,
+  PDB_TYPE_COMPLEX128,
+  PDB_TYPE_FIELDLIST,
+  PDB_TYPE_LABEL
+} pdb_type_kind_e;
+typedef U32 pdb_type_kind;
+
+typedef enum 
+{
+  PDB_TYPE_ATTRIB_HAS_NAME    = 1 << 0,
+  PDB_TYPE_ATTRIB_CONST       = 1 << 1,
+  PDB_TYPE_ATTRIB_VOLATILE    = 1 << 2,
+  PDB_TYPE_ATTRIB_UNALIGNED   = 1 << 3,
+  PDB_TYPE_ATTRIB_RESTRICTED  = 1 << 4,
+  PDB_TYPE_ATTRIB_LREF        = 1 << 5,
+  PDB_TYPE_ATTRIB_RREF        = 1 << 6,
+  PDB_TYPE_ATTRIB_FWDREF      = 1 << 7
+} pdb_type_attrib_e;
+typedef U32 pdb_type_attrib;
+
+typedef struct pdb_type_info_proc 
+{
+  /* Calling convention (pdb_cv_call_e) */
+  U32 conv;
+
+  /* itype of the return value. */
+  pdb_cv_itype ret_itype;
+
+  /* Arguments itype. */
+  pdb_cv_itype arg_itype;
+
+  U16 arg_count;
+} pdb_type_info_proc_t;
+
+typedef struct pdb_type_info_method 
+{
+  /* pdb_cv_call_e */
+  U32 conv;           
+
+  /* Return itype. */
+  pdb_cv_itype ret_itype;
+
+  /* Class itype to which this procedure belongs. */
+  pdb_cv_itype class_itype;    
+
+  /* Type of the "this" in C++ */
+  pdb_cv_itype this_itype;
+
+  /* Arguments itype. */
+  pdb_cv_itype arg_itype;      
+
+  U16 arg_count;
+} pdb_type_info_method;
+
+typedef struct pdb_type_info_methodlist
+{
+  pdb_pointer block; /* array of pdb_ml_method */
+} pdb_type_info_methodlist;
+
+typedef struct pdb_type_info_vtshape
+{
+  pdb_uint count;
+  /* pdb_lf_vts_desc are represented as 4bit packed integers */
+  pdb_pointer ptr; 
+} pdb_type_info_vtshape;
+
+typedef struct pdb_type_info_vftable
+{
+  pdb_cv_itype owner_itype;
+  pdb_cv_itype base_table_itype;
+  pdb_uint offset_in_object_layout;
+  pdb_pointer name;
+  pdb_pointer method_names;
+} pdb_type_info_vftable;
+
+typedef struct pdb_type_info_label
+{
+  pdb_cv_ptrmode mode;
+} pdb_type_info_label;
+
+typedef struct pdb_type_info_arglist 
+{
+  // Entry count
+  U32 count;
+
+  // Array of arguments, each argument is described with pdb_cv_itype
+  pdb_pointer itypes; 
+} pdb_type_info_arglist;
+
+typedef struct pdb_type_info_bitfield 
+{
+  /* Underlying type of a bitfield. */
+  pdb_cv_itype base_itype;
+
+  U8 len;
+  U8 pos;
+} pdb_type_info_bitfield;
+
+typedef struct pdb_type_info_udt 
+{
+  /* If this field is a UDT(User Defined Type) this field will have
+   * a non-zero value. */
+  pdb_cv_itype field_itype;
+
+  pdb_cv_itype dims_itype;
+
+  /* Used only to query ENUM_UDT fields. */
+  pdb_cv_itype base_itype;             
+
+  /* If this type belongs to a UDT this field will contains number
+   * of fields. */
+  U16 field_count;
+} pdb_type_info_udt;
+
+typedef struct pdb_type_info_funcid 
+{
+  pdb_pointer name;
+  pdb_cv_itype itype;
+} pdb_type_info_funcid;
+
+typedef struct pdb_type_info_mfuncid
+{
+  pdb_pointer name;          /* method name */
+  pdb_cv_itype parent_itype; /* type index of class (parent) */
+  pdb_cv_itype itype;        /* type index of method */
+} pdb_type_info_mfuncid;
+
+typedef struct pdb_type_info_stringid
+{
+  pdb_pointer   data;       /* string data */
+  pdb_cv_itemid sub_string; /* index of sub string */
+} pdb_type_info_stringid;
+
+typedef struct pdb_type_info_fieldlist
+{
+  pdb_pointer data;
+} pdb_type_info_fieldlist;
+
+typedef struct pdb_type_info_ptr
+{
+  pdb_cv_ptrmode mode;
+  pdb_cv_ptrtype type;
+  pdb_cv_ptr_attrib_t attr;
+} pdb_type_info_ptr;
+
+/*  You can call pdb_infer_itype and pass next_cv_itype to
+ *  resolve type information for next link in a chain of types.
+ *
+ *  If type was declared as 'void', pdb_infer_itype assigns zero to size.
+ *
+ *  When allocating storage for type's value - be careful user can declare any size. 
+ *  This is legal to compile "volatile char buffer[1*1024*1024*1024*1024]".
+ *  Make sure to put a cap on storage allocations to avoid out-of-memory error.
+ */
+typedef struct pdb_type 
+{
+  /* Size of the type in bytes.
+   *
+   * - When type is a pointer size will reflect an arch specific size,
+   *   so 4 bytes on x86 and 8 bytes on x64 and etc.
+   *   
+   * - When type is an array you progress through a chained-link
+   *   of dimensions; more info you can find
+   *   in pdb_tm_infer_ti, search for PDB_LF_ARRAY.
+   */
+  U32 size;
+
+  pdb_type_kind kind;
+
+  pdb_type_attrib attribs;
+
+  /* Unique id that is assigned to the type. */
+  pdb_cv_itype cv_itype;
+
+  /* Type is part of a hierarchy, this field indicates next itype. */
+  pdb_cv_itype next_cv_itype;
+
+  pdb_pointer name;
+
+  union {
+    /* Function ID */
+    pdb_type_info_funcid funcid;
+
+    /* Method ID */
+    pdb_type_info_mfuncid mfuncid;
+
+    /* String ID */
+    pdb_type_info_stringid stringid;
+
+    /* User-defined-type */
+    pdb_type_info_udt udt;
+
+    /* Procedure */
+    pdb_type_info_proc_t proc;
+
+    /* Pointer / Reference */
+    pdb_type_info_ptr ptr;
+
+    /* Method */
+    pdb_type_info_method method;
+
+    /* Virtual table shape */
+    pdb_type_info_vtshape vtshape;
+
+    /* Argument list */
+    pdb_type_info_arglist arglist;
+
+    /* Method list */
+    pdb_type_info_methodlist methodlist;
+
+    /* Bit-field */
+    pdb_type_info_bitfield bitfield;
+    
+    pdb_type_info_fieldlist fieldlist;
+
+    pdb_type_info_vftable vftable;
+
+    pdb_type_info_label label;
+  } u;
+} pdb_type;
+
+typedef struct pdb_tm_bucket 
+{
+  pdb_ti ti;
+  struct pdb_tm_bucket *next;
+} pdb_tm_bucket;
+
+typedef struct pdb_tm 
+{
+  struct pdb_context *pdb;
+  pdb_sn sn;
+  U32 *ti_offsets;
+  pdb_tm_bucket **buckets;
+  pdb_tm_header header;
+} pdb_tm;
+
+PDB_API U32
+pdb_calc_size_for_types(pdb_tm_header *tm_header);
+
+PDB_API syms_bool
+pdb_tm_init(pdb_tm *tm, struct pdb_context *pdb, pdb_default_stream sn, SymsArena *arena);
+
+PDB_API syms_bool
+pdb_tm_offset_for_ti(pdb_tm *tm, pdb_ti ti, U32 *ti_off_out);
+
+PDB_API syms_bool
+pdb_tm_get_itype_off(struct pdb_context *pdb, pdb_cv_itype itype, U32 *itype_off_out);
+
+PDB_API syms_bool
+pdb_tm_get_itemid_off(struct pdb_context *pdb, pdb_cv_itemid itemid, U32 *itemid_off);
+
+PDB_API syms_bool
+pdb_tm_find_ti(pdb_tm *tm, pdb_pointer *name, pdb_ti *ti_out);
+
+PDB_API syms_bool
+pdb_tm_find_itype(struct pdb_context *pdb, pdb_pointer name_ref, pdb_cv_itype *itype_out);
+
+PDB_API syms_bool
+pdb_tm_find_itemid(struct pdb_context *pdb, pdb_pointer name_ref, pdb_cv_itemid *itemid_out);
+
+PDB_API syms_bool
+pdb_tm_infer_ti(pdb_tm *tm, pdb_ti ti, pdb_type *type_out);
+
+PDB_API syms_bool 
+pdb_infer_itype(struct pdb_context *pdb, pdb_cv_itype itype, pdb_type *type_out);
+
+PDB_API syms_bool
+pdb_infer_itemid(struct pdb_context *pdb, pdb_cv_itemid itemid, pdb_type *type_out);
+
+PDB_API syms_bool
+pdb_type_from_name(struct pdb_context *pdb, const char *name, U32 name_len, pdb_type *type);
+
+typedef struct pdb_type_it 
+{
+  struct pdb_context *pdb;
+  pdb_cv_itype next_itype;
+  U32 type_map_index;
+  pdb_tm *type_map[2];
+} pdb_type_it;
+
+typedef struct 
+{
+  pdb_pointer file;
+  pdb_uint ln;
+  pdb_imod mod;
+} pdb_udt_srcline;
+
+PDB_API syms_bool
+pdb_find_udt_srcline(struct pdb_context *pdb, pdb_cv_itype itype, pdb_udt_srcline *srcline_out);
+
+//
+
+typedef struct pdb_cvdata_token 
+{
+  pdb_sn   sn;
+  pdb_uint soffset;
+} pdb_cvdata_token;
+
+typedef struct pdb_context 
+{
+  void *file_data;
+  pdb_uint file_size;
+
+  pdb_uint page_size;
+  pdb_uint free_page_map;
+  pdb_uint pages_used;
+  pdb_uint root_size;
+  U8 page_index_size;
+  union {
+    U16 *addr16;
+    U32 *addr32;
+  } page_map_addr;
+  pdb_uint ver;
+  pdb_uint time;
+  pdb_uint age;
+  /* NOTE(nick): Compiler writes into the PE header same GUID as here. 
+   * You can compare both to determine authenticity. */
+  SymsGUID auth_guid;
+  
+  pdb_sn typed_streams[PDB_TYPED_STREAM_MAX];
+
+  struct {
+    pdb_dbi_header header;
+
+    U32 ver;
+    U16 machine_type;
+
+    U32 modinfo_off;
+    U32 modinfo_len;
+
+    U32 seccon_off;
+    U32 seccon_len;
+
+    U32 secmap_off;
+    U32 secmap_len;
+
+    U32 fileinfo_off;
+    U32 fileinfo_len;
+
+    U32 ecinfo_off;
+    U32 ecinfo_len;
+
+    pdb_sn symrec_sn;
+    pdb_sn pubsym_sn;
+    pdb_sn globalsym_sn;
+
+    pdb_sn dbg_streams[PDB_DBG_STREAM_MAX];
+
+    U16 secs_num;
+    pdb_img_sec *secs;
+
+    /* NOTE(nick): Offsets into the PDB_DEFAULT_STREAM_DBI that point to modules. */
+    U32 mods_num;
+    pdb_uint *mods;
+  } dbi;
+
+  pdb_stream strtable;
+  pdb_stream stroffs;
+
+  struct pdb_tm tpi;
+  struct pdb_tm ipi;
+
+  U32 basic_typenames_array_num;
+  SymsString *basic_typenames_array;
+
+  U32 publics_array_num;
+  pdb_gsi_hr **publics_array;
+
+  U32 globals_array_num;
+  pdb_gsi_hr **globals_array;
+
+  // TODO(nick): trampolines are stored in a light-weight struct "pdb_cvdata_token". 
+  U32 trampoline_count;
+  syms_bool trampoline_contigous;
+  pdb_cvdata_token trampoline_data;
+} pdb_context;
+
+
+/* NOTE(nick; Oct 10; 2019): Offset into the string table. */
+typedef U32 pdb_stroff;
+
+PDB_API U32
+pdb_hashV1_bytes(const void *start, U32 cb, U32 mod);
+
+PDB_API U32
+pdb_hashV1_stream(pdb_stream *stream, U32 size, U32 mod);
+
+PDB_API U32
+pdb_hashV1_pointer(struct pdb_context *pdb, pdb_pointer *bytes, U32 mod);
+
+
+//
+
+typedef enum 
+{
+  PDB_ENCODED_LOCATION_NULL,
+  PDB_ENCODED_LOCATION_ENREGED,
+  PDB_ENCODED_LOCATION_REGREL,
+  PDB_ENCODED_LOCATION_RVA,
+  PDB_ENCODED_LOCATION_VA,
+  PDB_ENCODED_LOCATION_IMPLICIT
+} pdb_encoded_location_type_e;
+typedef U32 pdb_encoded_location_type;
+
+#define PDB_ENCODED_LOCATION_FLAG_NULL          0u
+#define PDB_ENCODED_LOCATION_FLAG_OUTSIDE_RANGE (1u << 0u)
+typedef U32 pdb_encoded_location_flags;
+
+typedef struct pdb_regrel 
+{
+  /* Offset that has to be applied to the register in order to get value for the local. */
+  U32 reg_off;
+
+  /* Index of the register (see pdb_cv_reg_e). */
+  U32 reg_index;
+} pdb_regrel;
+
+/* Compiler stored local in a register. */
+typedef struct pdb_enreged 
+{
+  U32 reg_index;
+} pdb_enreged;
+
+typedef struct pdb_persist 
+{
+  U64 va;
+} pdb_persist;
+
+#define PDB_LOCATION_IMPLICIT_VALUE_MAX 32
+typedef struct pdb_implicit 
+{
+  U8 len;
+  U8 data[PDB_LOCATION_IMPLICIT_VALUE_MAX];
+} pdb_implicit;
+
+/* Represents an encoded location of an address or a value.
+ * Use "pdb_decode_location" to resolve final address. */
+typedef struct pdb_encoded_location 
+{
+  /* Describes type of the memory that union below points to. */
+  pdb_encoded_location_type type;
+  pdb_encoded_location_flags flags;
+
+  union {
+    pdb_enreged  enreged;
+    pdb_regrel   regrel;
+    pdb_persist  persist;
+    pdb_implicit implicit;
+  } u;
+} pdb_encoded_location;
+
+typedef enum 
+{
+  PDB_LOCATION_NULL,
+  PDB_LOCATION_VA,
+  PDB_LOCATION_IMPLICIT
+} pdb_location_e;
+
+typedef struct pdb_location 
+{
+  pdb_location_e type;
+  union {
+    SymsAddr va;
+    pdb_implicit implicit;
+  } u;
+} pdb_location;
+
+PDB_API pdb_encoded_location
+pdb_encode_location_for_secoff(struct pdb_context *pdb, U32 sec, U32 off);
+
+PDB_API pdb_encoded_location
+pdb_encode_location_for_regrel(struct pdb_context *pdb, U32 regindex, U32 regoff);
+
+PDB_API pdb_encoded_location
+pdb_encode_location_for_datasym32(struct pdb_context *pdb, pdb_cv_datasym32 *datasym);
+
+#define PDB_MEMREAD_SIG(name) syms_bool name(void *context, SymsAddr va, void *read_buffer, U32 read_buffer_max)
+typedef PDB_MEMREAD_SIG(pdb_memread_sig);
+
+#define PDB_REGREAD_SIG(name) U32 name(void *context, U32 reg_index, void *read_buffer, U32 read_buffer_max)
+typedef PDB_REGREAD_SIG(pdb_regread_sig);
+
+#define PDB_REGWRITE_SIG(name) U32 name(void *context, U32 reg_index, void *write_buffer, U32 write_buffer_max)
+typedef PDB_REGWRITE_SIG(pdb_regwrite_sig);
+
+PDB_API syms_bool
+pdb_decode_location(pdb_encoded_location *encoded_loc,
+    SymsAddr orig_rebase,
+    SymsAddr rebase,
+    void *memread_ctx, pdb_memread_sig *memread,
+    void *regread_ctx, pdb_regread_sig *regread,
+    pdb_location *decoded_loc);
+
+PDB_API syms_bool
+pdb_build_va(struct pdb_context *pdb, U32 sec, U32 off, SymsAddr *out_va);
+
+PDB_API syms_bool
+pdb_build_sec_off(struct pdb_context *pdb, SymsAddr va, pdb_isec *sec, pdb_isec_umm *off);
+
+/* -------------------------------------------------------------------------------- */
+
+typedef struct pdb_var 
+{
+  /* Type of the local, you can pass this to the "pdb_infer_itype". */
+  pdb_cv_itype itype;
+
+  pdb_cv_localsym_flags flags;
+
+  /* Encoded address or value where variable is located in memory. */
+  pdb_encoded_location encoded_va;
+
+  pdb_pointer name;
+
+  pdb_pointer gaps;
+} pdb_var;
+
+typedef struct pdb_sec_it 
+{
+  struct pdb_context *pdb;
+  pdb_stream stream;
+} pdb_sec_it;
+
+PDB_API syms_bool
+pdb_sec_it_init(struct pdb_context *pdb, pdb_sec_it *sec_it);
+
+PDB_API syms_bool
+pdb_sec_it_next(pdb_sec_it *sec_it, pdb_img_sec *sec);
+
+PDB_API syms_bool
+pdb_sec_from_index(struct pdb_context *pdb, U32 index, pdb_img_sec *sec_out);
+
+typedef struct pdb_debug_sec 
+{
+  pdb_cv_ss_type type;
+  pdb_stream stream;
+} pdb_debug_sec;
+
+typedef enum 
+{
+  PDB_MOD_SEC_SYMS,
+  PDB_MOD_SEC_LINES_C11,
+  PDB_MOD_SEC_LINES_C13,
+  PDB_MOD_SEC_INLINE_LINES
+} pdb_mod_sec_type_e;
+
+typedef enum 
+{
+  PDB_MOD_FILECHKSUM_CACHED = (1 << 0)
+} pdb_mod_flags_e;
+typedef U32 pdb_mod_flags;
+
+typedef struct pdb_mod 
+{
+  struct pdb_context *pdb;
+
+  pdb_imod id;
+
+  pdb_mod_flags flags;
+
+  pdb_sn sn;
+
+  pdb_isec sec;
+  pdb_isec_umm sec_off;
+  pdb_isec_umm sec_len;
+
+  U32 syms_size;
+  U32 c11_lines_size;
+  U32 c13_lines_size;
+
+  pdb_pointer name;
+  pdb_pointer name2;
+
+  pdb_debug_sec filechksum;
+} pdb_mod;
+
+typedef struct pdb_mod_it 
+{
+  struct pdb_context *pdb;
+  pdb_stream dbi_data;
+  pdb_imod imod;
+} pdb_mod_it;
+
+PDB_API syms_bool
+pdb_mod_it_init(pdb_mod_it *mod_it, struct pdb_context *pdb);
+
+PDB_API syms_bool
+pdb_mod_it_next(pdb_mod_it *mod_it, pdb_mod *mod_out);
+
+PDB_API syms_bool
+pdb_mod_it_seek(pdb_mod_it *mod_it, pdb_imod imod);
+
+PDB_API syms_bool
+pdb_mod_get_debug_sec(pdb_mod *mod, pdb_mod_sec_type_e sec, pdb_stream *dbi_data_out);
+
+PDB_API syms_bool
+pdb_imod_from_isec(struct pdb_context *pdb, pdb_isec sec, pdb_isec_umm off, pdb_imod *imod_out);
+
+PDB_API syms_bool
+pdb_mod_init(pdb_mod *mod_out, struct pdb_context *pdb, pdb_imod imod);
+
+typedef struct pdb_file_info 
+{
+  pdb_pointer         path;
+  pdb_pointer         chksum;
+  pdb_cv_chksum_type  chksum_type;
+} pdb_file_info; // 72->56 bytes
+
+PDB_API syms_bool
+pdb_mod_get_filechksum(pdb_mod *mod, pdb_stream *stream_out);
+
+typedef struct pdb_debug_sec_it 
+{
+  pdb_stream stream;
+} pdb_debug_sec_it;
+
+PDB_API syms_bool
+pdb_debug_sec_it_init(pdb_debug_sec_it *it, pdb_mod *mod);
+
+PDB_API syms_bool
+pdb_debug_sec_it_next(pdb_debug_sec_it *it, pdb_debug_sec *sec_out);
+
+typedef struct pdb_dss_it 
+{
+  pdb_cv_ss_type type;
+  pdb_stream stream;
+  syms_bool ex_mode;
+} pdb_dss_it;
+
+/* Inline data is stored in DBI stream and this is header file for it. */
+typedef struct pdb_ss_inline 
+{
+  pdb_cv_itemid inlinee;
+
+  U32 src_ln;
+  U32 file_id;
+  U32 extra_files_count;
+  pdb_pointer extra_files;
+} pdb_ss_inline;
+
+PDB_API syms_bool
+pdb_dss_it_init(pdb_dss_it *it_out, pdb_debug_sec *sec);
+
+PDB_API syms_bool
+pdb_dss_it_next_inline(pdb_dss_it *it, pdb_ss_inline *inline_out);
+
+typedef struct pdb_sym_it 
+{
+  syms_bool inited_from_token;
+  pdb_stream stream;
+} pdb_sym_it;
+
+PDB_API syms_bool
+pdb_sym_it_init(pdb_sym_it *sym_it, pdb_mod *mod);
+
+PDB_API syms_bool
+pdb_sym_it_init_token(pdb_sym_it *sym_it, struct pdb_context *pdb, pdb_cvdata_token token);
+
+PDB_API syms_bool
+pdb_sym_it_read(pdb_sym_it *sym_it, pdb_cv_sym_type *type_out, pdb_stream *stream_out);
+
+PDB_API syms_bool
+pdb_sym_it_peek(pdb_sym_it *sym_it, pdb_cv_sym_type *type_out, pdb_stream *stream_out);
+
+PDB_API syms_bool
+pdb_sym_it_next(pdb_sym_it *sym_it, pdb_cv_sym_type *type_out, pdb_stream *stream_out);
+
+enum
+{
+  PDB_LOCAL_EXPORT_NULL,
+  PDB_LOCAL_EXPORT_VAR,
+  PDB_LOCAL_EXPORT_SCOPE,
+  PDB_LOCAL_EXPORT_SCOPE_END
+};
+typedef pdb_uint pdb_local_export_type;
+
+typedef struct pdb_scope
+{
+  SymsAddr inst_lo;
+  SymsAddr inst_hi;
+} pdb_scope;
+
+typedef struct pdb_local_export
+{
+  pdb_local_export_type type;
+  union
+  {
+    pdb_var var;
+    pdb_scope scope;
+  } u;
+} pdb_local_export;
+
+typedef struct pdb_local_it 
+{
+  struct pdb_context *pdb;
+  syms_bool defrange_mode;
+  pdb_sym_it sym_it;
+
+  pdb_isec sec;
+  pdb_isec_umm sec_off;
+
+  U32 range_off;
+  U32 range_len;
+
+  S32 block32_count;
+  S32 regrel32_count;
+  S32 inlinesite_count;
+
+  U16 arg_count;
+} pdb_local_it;
+
+PDB_API syms_bool
+pdb_local_it_init_(pdb_local_it *local_it, struct pdb_context *pdb, pdb_isec sec, pdb_isec_umm sec_off, pdb_cvdata_token cvdata);
+
+PDB_API syms_bool
+pdb_local_it_init(pdb_local_it *local_it, struct pdb_context *pdb, pdb_isec sec, pdb_isec_umm sec_off);
+
+PDB_API syms_bool
+pdb_local_it_next(pdb_local_it *local_it, pdb_local_export *export_out);
+
+typedef struct pdb_global_it 
+{
+  struct pdb_context *pdb;
+  pdb_stream stream;
+  U32 hr_index;
+  pdb_gsi_hr *hr;
+} pdb_global_it;
+
+PDB_API syms_bool
+pdb_global_it_init(pdb_global_it *global_it, struct pdb_context *pdb);
+
+PDB_API syms_bool
+pdb_global_it_next(pdb_global_it *global_it, pdb_var *var_out);
+
+PDB_API syms_bool
+pdb_global_from_name(struct pdb_context *pdb, const char *name, U32 name_size, pdb_var *global_out);
+
+typedef struct pdb_file_it 
+{
+  struct pdb_context *pdb;
+  pdb_stream strtable;
+  pdb_stream stroffs;
+  U32 off_count;
+  U32 num_read;
+} pdb_file_it;
+
+PDB_API syms_bool
+pdb_file_it_init(pdb_file_it *file_it, struct pdb_context *pdb);
+
+PDB_API syms_bool
+pdb_file_it_next(pdb_file_it *file_it, pdb_pointer *file_out);
+
+typedef struct pdb_proc_it 
+{
+  pdb_sym_it sym_it;
+} pdb_proc_it;
+
+typedef struct pdb_proc 
+{
+  pdb_isec sec;
+  pdb_isec_umm sec_off;
+
+  /* Length in bytes, of all instruction that contribute to a procedure. */
+  pdb_uint size;
+
+  pdb_cv_itype itype;
+
+  /* pdb_cv_proc_flags_e */
+  pdb_uint flags;
+
+  pdb_string_ref name;
+
+  pdb_cvdata_token cvdata;
+} pdb_proc;
+
+PDB_API syms_bool
+pdb_proc_it_init(pdb_proc_it *proc_it, pdb_mod *mod);
+
+PDB_API syms_bool
+pdb_proc_it_next(pdb_proc_it *proc_it, pdb_proc *out_proc);
+
+PDB_API syms_bool
+pdb_proc_from_name(pdb_context *pdb, const char *name, pdb_uint name_len, pdb_proc *proc_out);
+
+PDB_API syms_bool
+pdb_proc_from_cvdata(pdb_context *pdb, pdb_cvdata_token cvdata, pdb_proc *proc_out);
+
+PDB_API syms_bool
+pdb_proc_from_stream(pdb_cv_sym_type cvtype, pdb_stream *cvdata, pdb_proc *proc_out);
+
+PDB_API syms_bool
+pdb_proc_from_name_(pdb_context *pdb, pdb_pointer *name, pdb_proc *proc_out);
+
+typedef struct pdb_inline_site 
+{
+  pdb_pointer name;
+
+  pdb_cv_itype itype;
+
+  pdb_isec sec;
+  pdb_isec_umm sec_off;
+  u32 size;
+
+  pdb_file_info fi;
+
+  U32 ln_at_pc;
+
+  pdb_cvdata_token cvdata;
+} pdb_inline_site;
+
+typedef struct pdb_inline_it 
+{
+  S32 site_count;
+  /* NOTE(nick; Oct 6 2019): Procedure relative Program Counter. */
+  SymsAddr proc_pc;
+  pdb_proc proc;
+  pdb_mod mod;
+  pdb_sym_it sym_it;
+} pdb_inline_it;
+
+PDB_API syms_bool
+pdb_inline_it_init(pdb_inline_it *inline_it, struct pdb_context *pdb, pdb_cvdata_token cvproc, SymsAddr proc_rva);
+
+PDB_API syms_bool
+pdb_inline_it_read(pdb_inline_it *inline_it, pdb_cv_inlinesym *sym_out, pdb_stream *ba_out);
+
+PDB_API syms_bool
+pdb_inline_it_next(pdb_inline_it *inline_it, pdb_inline_site *site_out);
+
+typedef struct pdb_const_value 
+{
+  pdb_string_ref name;
+  pdb_cv_itype itype;
+  pdb_numeric value;
+} pdb_const_value;
+
+typedef struct pdb_const_it 
+{
+  struct pdb_context *pdb;
+  pdb_stream stream;
+  U32 index;
+  pdb_gsi_hr *hr;
+} pdb_const_it;
+
+PDB_API syms_bool
+pdb_const_it_init(pdb_const_it *const_it, struct pdb_context *pdb);
+
+PDB_API syms_bool
+pdb_const_it_next(pdb_const_it *const_it, pdb_const_value *const_out);
+
+PDB_API syms_bool
+pdb_const_from_name(struct pdb_context *pdb, const char *name, U32 name_size, pdb_const_value *const_out);
+
+typedef struct pdb_arg_it 
+{
+  /* parser context */
+  struct pdb_context *pdb;
+
+  /* current argument index */
+  u32 idx;
+
+  /* array of pdb_cv_itype entries */
+  pdb_pointer itypes;
+} pdb_arg_it;
+
+PDB_API syms_bool
+pdb_arg_it_init(pdb_arg_it *it, struct pdb_context *pdb, pdb_cv_itype arglist_itype);
+
+PDB_API syms_bool
+pdb_arg_it_next(pdb_arg_it *it, pdb_cv_itype *out_arg_type);
+
+/* -------------------------------------------------------------------------------- */
+
+enum 
+{
+  PDB_MEMBER_TYPE_NULL,
+
+  /* Member for data that is stored inside the UDT  */
+  PDB_MEMBER_TYPE_DATA,
+
+  /* Class method */
+  PDB_MEMBER_TYPE_METHOD,
+
+  /* Enum member */
+  PDB_MEMBER_TYPE_ENUMERATOR,
+
+  /* Type that was declared inside a class */
+  PDB_MEMBER_TYPE_NESTED_TYPE,
+
+  /* Inherited class */
+  PDB_MEMBER_TYPE_BASE_CLASS,
+
+  /* Describes memory region for virtual functions.
+   * Memory is packed as an array of 4 bit integers,
+   * each integer maps into pdb_lf_vts_desc. */
+  PDB_MEMBER_TYPE_VIRTUAL_TABLE,
+  
+  /* Static data member of a class */
+  PDB_MEMBER_TYPE_STATIC_DATA,
+
+  /* 
+   * CodeView packs methods with matching identifiers into 
+   * an array of methods. I am sorry for ugly hack - we need 
+   * it to hide in the interface layer codeview shenanigans. 
+   *
+   * Also we dont have an allocator in the API layer, 
+   * iterators, symbols, line info, and so on are allocated 
+   * on the stack. Library allocates memory just for lookup
+   * acceleration and it is possible to parse entire pdb
+   * with zero memory allocations! 
+   *
+   * */
+  PDB_MEMBER_TYPE_METHODLIST
+};
+typedef u32 pdb_member_type;
+
+/* Contains anything that resembles data member in struct, class, enum, and union */
+typedef struct pdb_member 
+{
+  pdb_pointer name;
+
+  pdb_member_type type;
+
+  pdb_cv_itype itype;
+
+  pdb_cv_fldattr_t attr;
+
+  union 
+  {
+    struct
+    {
+      /* Parent-relative offset where this member is located in memory. */
+      U32 offset; 
+    } data;
+
+    struct 
+    {
+      u32 count;
+    } methodlist;
+
+    struct
+    {
+      pdb_uint vbaseoff;
+    } method;
+
+    struct 
+    {
+      /* Use this field only if the "type" field is set to the PDB_MEMBER_TYPE_ENUMERATOR.
+       * This field indicates value that was assigned to enum field. */
+      pdb_numeric value;
+    } enumerator;
+
+    struct 
+    {
+      u32 offset;
+    } base_class;
+
+    struct
+    {
+      u32 offset;
+    } vtab;
+  } u;
+} pdb_member;
+
+typedef struct pdb_member_it 
+{
+  U32 stream_end;
+  struct pdb_context *pdb;
+  pdb_type_info_udt *udt;
+  pdb_stream stream;
+} pdb_member_it;
+
+PDB_API syms_bool
+pdb_member_it_init(pdb_member_it *member_it, struct pdb_context *pdb, pdb_type_info_udt *udt);
+
+PDB_API syms_bool
+pdb_member_it_next(pdb_member_it *member_it, pdb_member *member_out);
+
+/* -------------------------------------------------------------------------------- */
+
+/* Indicates if we need to read another line section. */
+#define PDB_LINE_IT_FLAGS_NEW_SECTION (1u << 0u)
+typedef U32 pdb_line_it_flags;
+
+typedef enum 
+{
+  PDB_LINE_FORMAT_NULL,
+  PDB_LINE_FORMAT_C11,
+  PDB_LINE_FORMAT_C13
+} pdb_line_format;
+
+typedef struct pdb_line_it 
+{
+  /* Iterator for the modules. */
+  pdb_mod mod;
+  pdb_file_info fi;
+  pdb_line_format format;
+  pdb_line_it_flags flags;
+  pdb_stream stream;
+
+  U16 sec;
+  U32 sec_off;
+  U32 sec_size;
+
+  U32 last_read_ln;
+
+  union {
+    struct {
+      /* Count for the total number of lines in current section. */
+      U32 line_index_max;
+
+      /* When a new module comes, we scan it to get offset of the checksum 
+       * section because all lines instead of providing offset where the 
+       * file name is in the name-stream, they give offset into the 
+       * checksum section and from there we have to resolve the file name
+       * for given lines block. */
+      U32 chksum_off;
+
+      /* Offset where current line section ends. */
+      U32 sec_end;
+
+      pdb_stream cv_lines;
+    } c13;
+
+    struct {
+      U32 file_index;
+      U32 file_count;
+
+      U16 filesec_index;
+      U16 filesec_count;
+
+      U16 pair_index;
+      U16 pair_count;
+
+      pdb_stream secrange_stream;
+      pdb_stream off_stream;
+      pdb_stream ln_stream;
+    } c11;
+  } u;
+} pdb_line_it;
+
+typedef struct pdb_line 
+{
+  U16 sec;
+  U32 off;
+  U32 ln;
+} pdb_line;
+
+PDB_API syms_bool
+pdb_line_it_init(pdb_line_it *line_it, pdb_mod *mod);
+
+PDB_API syms_bool
+pdb_line_it_next(pdb_line_it *line_it, pdb_line *line);
+
+typedef struct pdb_map 
+{
+  /* section and offset for the first instruction in the source line. */
+  pdb_isec sec;
+  pdb_isec_umm sec_off;
+
+  /* Size of all instructions that contribute to the source line. */
+  U32 size;
+
+  /* Line number to which VA was mapped. */
+  U32 ln;
+
+  pdb_file_info fi;
+} pdb_map;
+
+PDB_API syms_bool
+pdb_src_to_va(struct pdb_context *pdb, char *file_name, U32 file_name_len, U32 ln,
+        pdb_map *map);
+
+PDB_API syms_bool
+pdb_va_to_src(struct pdb_context *pdb, SymsAddr line_va, pdb_map *map);
+
+/* -------------------------------------------------------------------------------- */
+
+PDB_API syms_bool
+pdb_find_nearest_sym(struct pdb_context *pdb, SymsAddr va, pdb_pointer *name_out);
+
+PDB_API syms_bool
+pdb_find_nearest_sc(struct pdb_context *pdb, SymsAddr va, pdb_sc *sc_out);
+
+PDB_API syms_bool
+pdb_find_nearest_sc_ex(struct pdb_context *pdb, U32 sec, U32 off, pdb_sc *sc_out);
+
+// File Info section of DBI packs unique strings of source paths. It has a 16bit counter
+// -- hello from good old days -- and today medium size program that link with CRT
+// have 30K files when you unpack file info section.
+typedef struct pdb_fileinfo
+{
+  u32 mod_count;
+  u32 src_count;
+  pdb_stream imod_block;        // pdb_uint16[mod_count] module indices
+  pdb_stream count_block;       // source count per module, counts are packed as variable-length arrays.
+  pdb_stream ich_block;         // offsets into str_block
+  pdb_stream str_block;         // char * array
+} pdb_fileinfo;
+
+PDB_API syms_bool
+pdb_fileinfo_init(struct pdb_context *pdb, pdb_fileinfo *fi);
+
+PDB_API syms_bool
+pdb_fileinfo_get_src_count(pdb_fileinfo *fi, pdb_imod imod, pdb_uint *count_out);
+
+PDB_API syms_bool
+pdb_fileinfo_lookup(pdb_fileinfo *fi, pdb_imod imod, pdb_uint index, pdb_pointer *cptr_out);
+
+typedef struct pdb_strtable
+{
+  u32 magic;
+  u32 version;
+  u32 bucket_count;
+  pdb_stream strblock;
+  pdb_stream buckets;
+} pdb_strtable;
+
+PDB_API syms_bool
+pdb_init_strtable(struct pdb_strtable *st, struct pdb_stream *stream);
+
+PDB_API pdb_stream
+pdb_get_strtable(struct pdb_context *pdb);
+
+PDB_API pdb_stream
+pdb_get_stroffs(struct pdb_context *pdb);
+
+PDB_API syms_bool
+pdb_strtable_at(struct pdb_context *pdb, pdb_ni ni, pdb_pointer *str_out);
+
+PDB_API syms_bool
+pdb_strtable_off_to_str(struct pdb_context *pdb, pdb_stroff stroff, pdb_pointer *str_out);
+
+PDB_API syms_bool
+pdb_strtable_find(struct pdb_context *pdb, pdb_pointer *name, pdb_strcmp_flags_e cmp_flags, pdb_ni *ni);
+
+PDB_API syms_bool
+pdb_strtable_find_str(struct pdb_context *pdb, SymsString name, pdb_strcmp_flags_e cmp_flags, pdb_ni *ni);
+
+
+PDB_API U32
+pdb_parse_gsi_hash_table(pdb_stream *gsi_stream, SymsArena *arena, pdb_gsi_hr ***out_table, U32 *out_num_table);
+
+PDB_API syms_bool
+pdb_init(struct pdb_context *pdb, void *data, U32 data_size);
+
+PDB_API SymsUMM
+pdb_load_types(pdb_context *pdb, struct SymsArena *arena);
+
+PDB_API SymsUMM
+pdb_load_publics(pdb_context *pdb, struct SymsArena *arena);
+
+PDB_API SymsUMM
+pdb_load_globals(pdb_context *pdb, struct SymsArena *arena);
+
+PDB_API SymsUMM
+pdb_load_dbi(pdb_context *pdb, struct SymsArena *arena);
+
+PDB_API SymsUMM
+pdb_load_trampolines(pdb_context *pdb, struct SymsArena *arena);
+
+PDB_API SymsUMM
+pdb_calc_size_for_aux_data(pdb_context *pdb);
+
+PDB_API syms_bool
+pdb_load_aux_data(pdb_context *pdb, SymsArena *arena);
+
+PDB_API syms_bool
+pdb_trampoline_from_ip(pdb_context *pdb, pdb_isec sec, pdb_isec_umm sec_off, pdb_isec *sec_out, pdb_isec_umm *off_out);
+
+PDB_API syms_bool
+pdb_stream_read_numeric(pdb_stream *stream, pdb_numeric *out_num);
+
+PDB_API syms_bool
+pdb_stream_read_numeric_u32(pdb_stream *stream, U32 *out_value);
+
+PDB_API syms_bool
+pdb_stream_read_numeric_u64(pdb_stream *stream, U64 *out_value);
+
+PDB_API SymsNTFileHeaderMachineType
+pdb_get_machine_type(struct pdb_context *pdb);
+
+PDB_API syms_bool
+pdb_type_it_init(pdb_type_it *it, struct pdb_context *pdb);
+
+PDB_API syms_bool
+pdb_type_it_next(pdb_type_it *it, pdb_cv_itype *itype_out);
+
+PDB_API const char *
+pdb_ver_to_str(pdb_context *pdb);
+
+PDB_API const char *
+pdb_dbi_ver_to_str(pdb_context *pdb);
+
+PDB_API const char *
+pdb_basic_itype_to_str(pdb_basic_type itype);
+
+PDB_API syms_bool
+pdb_fileinfo_get_strblock(pdb_fileinfo *fi, pdb_imod imod, u32 *count_out, pdb_stream *strblock_out);
+
+#endif /* SYMS_PDB_INCLUDE_H */
diff --git a/Engine/Source/Runtime/SymsLib/src/pdb/syms_pdb_api.c b/Engine/Source/Runtime/SymsLib/src/pdb/syms_pdb_api.c
new file mode 100644
--- /dev/null
+++ b/Engine/Source/Runtime/SymsLib/src/pdb/syms_pdb_api.c
@@ -0,0 +1,1286 @@
+// Copyright Epic Games, Inc. All Rights Reserved.
+/******************************************************************************
+ * File   : syms_pdb.c                                                        *
+ * Author : Nikita Smith                                                      *
+ * Created: 2020/06/20                                                        *
+ * Purpose: API wrappers and type conversion for pdb.h                        *
+ ******************************************************************************/
+
+SYMS_INTERNAL SymsStringRef
+syms_string_ref_null(void);
+
+SYMS_INTERNAL pdb_cv_itype
+syms_typeid_to_pdb(const SymsTypeID *type_id)
+{
+  pdb_cv_itype itype = *(const pdb_cv_itype *)type_id->impl;
+  return itype;
+}
+
+SYMS_INTERNAL SymsTypeID
+syms_typeid_for_pdb(pdb_cv_itype itype)
+{
+  SymsTypeID type_id;
+  syms_memzero(&type_id, sizeof(type_id));
+  type_id.kind = SYMS_TYPE_ID_PDB;
+  *(pdb_cv_itype *)type_id.impl = itype;
+  return type_id;
+}
+
+SYMS_INTERNAL syms_int
+syms_typeid_cmp_pdb(pdb_cv_itype a, pdb_cv_itype b)
+{
+  syms_int r = 0;
+  if (a < b) {
+    r = -1;
+  }
+  if (a > b) {
+    r = +1;
+  }
+  return r;
+}
+
+// ----------------------------------------------------------------------------
+
+SYMS_INTERNAL syms_bool
+syms_mod_iter_init_pdb(SymsDebugInfoPdb *debug_info, SymsModIterPdb *iter)
+{
+  // TODO: error handling
+  if (pdb_mod_it_init(&iter->impl, &debug_info->context)) {
+    return SYMS_ERR_OK;
+  } else {
+    syms_memzero(iter, sizeof(*iter));
+    return SYMS_ERR_INVALID_CODE_PATH;
+  }
+}
+
+SYMS_INTERNAL syms_bool
+syms_mod_iter_next_pdb(SymsDebugInfoPdb *debug_info, SymsModIterPdb *iter, SymsMod *mod_out)
+{
+  pdb_context *context = &debug_info->context;
+  pdb_mod_it *impl = &iter->impl;
+  pdb_mod *mod = (pdb_mod *)mod_out->impl;
+  syms_bool result = syms_false;
+  if (pdb_mod_it_next(impl, mod)) {
+    SymsAddr p = 0;
+    mod_out->id = mod->id;
+    if (pdb_build_va(context, mod->sec, mod->sec_off, &p)) {
+      mod_out->va = (SymsAddr)p;
+      mod_out->size = mod->sec_len;
+    } else {
+      mod_out->va = 0;
+      mod_out->size = 0;
+    }
+    mod_out->name = syms_string_ref_pdb(&mod->name);
+    result = syms_true;
+  }
+  return result;
+}
+
+// -----------------------------------------------------------------------------
+
+SYMS_INTERNAL SymsLineIterPdb *
+syms_line_iter_to_pdb(SymsLineIter *iter)
+{
+  SymsLineIterPdb *impl = (SymsLineIterPdb *)iter->impl;
+  return impl;
+}
+
+
+SYMS_INTERNAL syms_bool
+syms_line_iter_init_pdb(SymsDebugInfoPdb *debug_info, SymsLineIterPdb *iter, pdb_mod *mod, syms_bool *has_line_count)
+{ (void)debug_info;
+  *has_line_count = syms_true;
+  return pdb_line_it_init(&iter->impl, mod);
+}
+
+SYMS_INTERNAL syms_bool
+syms_line_iter_next_pdb(SymsDebugInfoPdb *debug_info, SymsLineIterPdb *iter, SymsLine *line_out, syms_bool *switched_file, SymsSourceFile *file, syms_uint *line_count)
+{
+  syms_bool keep_fetching = syms_true;
+  syms_bool result = syms_false;
+  for (; keep_fetching ;) {
+    pdb_line_it *impl = &iter->impl;
+    pdb_line line;
+
+    keep_fetching = syms_false;
+
+    if (pdb_line_it_next(impl, &line)) {
+      if (impl->flags & PDB_LINE_IT_FLAGS_NEW_SECTION) {
+        file->name = syms_string_ref_pdb(&impl->fi.path);
+        switch (impl->fi.chksum_type) {
+        case PDB_CV_CHECKSUM_NULL:     file->chksum_type = SYMS_CHECKSUM_NULL;   break;
+        case PDB_CV_CHECKSUM_MD5:      file->chksum_type = SYMS_CHECKSUM_MD5;    break;
+        case PDB_CV_CHECKSUM_SHA1:     file->chksum_type = SYMS_CHECKSUM_SHA1;   break;
+        case PDB_CV_CHECKSUM_SHA256:   file->chksum_type = SYMS_CHECKSUM_SHA256; break;
+        default:                       SYMS_ASSERT_FAILURE("unknown checksum type");  break;
+        }
+        pdb_pointer_read(&debug_info->context, &impl->fi.chksum, 0, &file->chksum[0], sizeof(file->chksum));
+
+        switch (impl->format) {
+        case PDB_LINE_FORMAT_C11: *line_count = impl->u.c11.pair_count; break;
+        case PDB_LINE_FORMAT_C13: *line_count = impl->u.c13.line_index_max; break;
+        default: SYMS_INVALID_CODE_PATH;
+        }
+
+        *switched_file = syms_true;
+      }
+
+      {
+        SymsAddr va = 0;
+        pdb_build_va(&debug_info->context, line.sec, line.off, &va);
+        syms_line_init(line_out, va, line.ln, 0);
+      }
+      result = syms_true;
+    }
+  }
+  return result;
+}
+
+// -----------------------------------------------------------------------------
+
+SYMS_INTERNAL syms_bool
+syms_member_iter_init_pdb(SymsDebugInfoPdb *debug_info, SymsMemberIterPdb *iter, SymsType *type)
+{
+  pdb_type *impl = (pdb_type *)type->impl;
+  iter->methodlist_name = pdb_pointer_bake_null();
+  iter->methodlist_index = 0;
+  iter->methodlist_count = 0;
+  iter->methodlist_block = pdb_pointer_bake_null();
+  return pdb_member_it_init(&iter->impl, &debug_info->context, &impl->u.udt);
+}
+
+SYMS_INTERNAL syms_bool
+pdb_export_leaf_member(pdb_member *member, SymsMember *member_out)
+{
+  syms_uint cv_mprop;
+
+  /* zero out output */
+  member_out->type      = SYMS_MEMBER_TYPE_NULL;
+  member_out->modifier  = SYMS_MEMBER_MODIFIER_NULL;
+  member_out->access    = SYMS_MEMBER_ACCESS_NULL;
+  member_out->type_id   = syms_typeid_for_pdb(member->itype);
+  member_out->name_ref  = syms_string_ref_pdb(&member->name);
+
+  /* in c++ these represent private, public, protected members of class or struct */
+  switch (PDB_CV_FLDATTR_ACCESS_MASK(member->attr)) {
+  case PDB_CV_FLDATTR_ACCESS_PRIVATE:   member_out->access = SYMS_MEMBER_ACCESS_PRIVATE;    break;
+  case PDB_CV_FLDATTR_ACCESS_PUBLIC:    member_out->access = SYMS_MEMBER_ACCESS_PUBLIC;     break;
+  case PDB_CV_FLDATTR_ACCESS_PROTECTED: member_out->access = SYMS_MEMBER_ACCESS_PROTECTED;  break;
+  default:                              member_out->access = SYMS_MEMBER_ACCESS_NULL;       break;
+  }
+
+  /* method modifiers */
+  cv_mprop = PDB_CV_FLDATTR_MPROP_MASK(member->attr);
+  switch (cv_mprop) {
+  case PDB_CV_FLDATTR_MPROP_VANILLA:    member_out->modifier = SYMS_MEMBER_MODIFIER_VANILLA;      break;
+  case PDB_CV_FLDATTR_MPROP_STATIC:     member_out->modifier = SYMS_MEMBER_MODIFIER_STATIC;       break;
+  case PDB_CV_FLDATTR_MPROP_FRIEND:     member_out->modifier = SYMS_MEMBER_MODIFIER_FRIEND;       break;
+  case PDB_CV_FLDATTR_MPROP_VIRTUAL:    member_out->modifier = SYMS_MEMBER_MODIFIER_VIRTUAL;      break;
+  case PDB_CV_FLDATTR_MPROP_INTRO:      member_out->modifier = SYMS_MEMBER_MODIFIER_INTRO;        break;
+  case PDB_CV_FLDATTR_MPROP_PUREVIRT:   member_out->modifier = SYMS_MEMBER_MODIFIER_PURE_VIRTUAL; break;
+  case PDB_CV_FLDATTR_MPROP_PUREINTRO:  member_out->modifier = SYMS_MEMBER_MODIFIER_PURE_INTRO;   break;
+  default:                              member_out->modifier = SYMS_MEMBER_MODIFIER_NULL;         break;
+  }
+
+  switch (member->type) {
+  case PDB_MEMBER_TYPE_DATA: {
+    member_out->type = SYMS_MEMBER_TYPE_DATA;
+    member_out->u.data_offset = (SymsUMM)member->u.data.offset;
+  } break;
+
+  case PDB_MEMBER_TYPE_STATIC_DATA: {
+    member_out->type = SYMS_MEMBER_TYPE_STATIC_DATA;
+  } break;
+
+  case PDB_MEMBER_TYPE_ENUMERATOR: {
+    member_out->type = SYMS_MEMBER_TYPE_ENUM;
+    syms_memcpy(member_out->u.enum_value, member->u.enumerator.value.u.data, sizeof(member->u.enumerator.value.u.data));
+  } break;
+
+  case PDB_MEMBER_TYPE_VIRTUAL_TABLE: {
+    member_out->type = SYMS_MEMBER_TYPE_VTABLE;
+    member_out->u.vtab.offset = member->u.vtab.offset;
+  } break;
+
+  case PDB_MEMBER_TYPE_METHOD: {
+    member_out->type = SYMS_MEMBER_TYPE_METHOD;
+    member_out->u.method.vbaseoff = member->u.method.vbaseoff;
+  } break;
+
+  case PDB_MEMBER_TYPE_BASE_CLASS: {
+    member_out->type = SYMS_MEMBER_TYPE_BASE_CLASS;
+    member_out->u.base_class.offset = member->u.base_class.offset;
+  } break;
+
+  case PDB_MEMBER_TYPE_NESTED_TYPE: {
+    member_out->type = SYMS_MEMBER_TYPE_NESTED_TYPE;
+  } break;
+
+  case PDB_MEMBER_TYPE_METHODLIST: {
+    /* processed in the member iterator */
+  } break;
+  }
+
+  return syms_true;
+}
+
+SYMS_INTERNAL syms_bool
+syms_member_iter_next_pdb(SymsDebugInfoPdb *debug_info, SymsMemberIterPdb *iter, SymsMember *member_out)
+{
+  syms_bool is_parsed = syms_false;
+  pdb_context *pdb = &debug_info->context;
+  pdb_member_it *impl = &iter->impl;
+  pdb_member member;
+
+export_methodlist:;
+  if (iter->methodlist_index < iter->methodlist_count) {
+    /* TODO: leaking PDB code into API layer -- refactor! */
+
+    pdb_ml_method method;
+    pdb_uint read_size;
+    pdb_uint read_offset;
+
+    /* Update index so we know when to stop parsing. */
+    iter->methodlist_index += 1;
+
+    /* Read out one entry from method array. */
+    read_offset = 0;
+    read_size = pdb_pointer_read(pdb, &iter->methodlist_block, read_offset, &method, sizeof(method));
+    read_offset += read_size;
+    if (read_size == sizeof(method)) {
+      pdb_uint mprop;
+      pdb_member temp_member;
+      pdb_uint vbaseoff;
+
+      /* Optional virtual table offset. Methods that are declared 
+       * as virtual and have no implementation - dont have an offset. */
+      mprop = PDB_CV_FLDATTR_MPROP_MASK(method.attr);
+      vbaseoff = 0;
+      if (mprop == PDB_CV_FLDATTR_MPROP_PUREINTRO || mprop == PDB_CV_FLDATTR_MPROP_INTRO) {
+        read_size = pdb_pointer_read(pdb, &iter->methodlist_block, read_offset, &vbaseoff, sizeof(vbaseoff));
+        if (read_size == sizeof(vbaseoff)) {
+          read_offset += read_size;
+        } else {
+          SYMS_ASSERT_FAILURE_PARANOID("unable to read virtual table offset for method");
+        }
+      }
+
+      /* Fill out pdb version of member and export it as generic member. */
+      temp_member.name = iter->methodlist_name;
+      temp_member.type = PDB_MEMBER_TYPE_METHOD;
+      temp_member.itype = method.index;
+      temp_member.attr = method.attr;
+      temp_member.u.method.vbaseoff = vbaseoff;
+      is_parsed = pdb_export_leaf_member(&temp_member, member_out);
+    }
+   
+    if ( ! is_parsed) {
+      /* Unable to prase method list, but we dont want to give up and return an error
+       * because there could be valid members following this one. */
+      iter->methodlist_index = iter->methodlist_count;
+    }
+  }
+
+  if ( ! is_parsed && pdb_member_it_next(impl, &member)) {
+    if (member.type == PDB_MEMBER_TYPE_METHODLIST) {
+      pdb_type type;
+
+      if ( ! pdb_infer_itype(pdb, member.itype, &type)) {
+        /* corrupted/invalid itype. */
+        SYMS_ASSERT_FAILURE_PARANOID("corrupted method list itype");
+        return syms_false;
+      }
+      if (type.kind != PDB_TYPE_METHODLIST) {
+        /* inferred type is not a method list. */
+        SYMS_ASSERT_FAILURE_PARANOID("expected a method list type");
+        return syms_false;
+      }
+
+      /* HACK: Method lists are arrays with methods that have identical names.
+       * And since iterator exports one method at a time we have to work around here.
+       * Whenever a method list is encountered we pretend that this iterator
+       * is a for loop, set method index to zero, set count, and jump.
+       * I know that goto is bad code but I hvae no idea how to work around 
+       * this problem without an allocator.
+       */
+      iter->methodlist_name = member.name;
+      iter->methodlist_index = 0;
+      iter->methodlist_count = member.u.methodlist.count;
+      iter->methodlist_block = type.u.methodlist.block;
+      goto export_methodlist;
+    } else {
+      is_parsed = pdb_export_leaf_member(&member, member_out);
+    }
+  }
+
+  return is_parsed;
+}
+
+SYMS_INTERNAL syms_bool
+syms_global_iter_init_pdb(SymsDebugInfoPdb *debug_info, SymsGlobalIterPdb *iter)
+{
+  return pdb_global_it_init(&iter->impl, &debug_info->context);
+}
+
+SYMS_INTERNAL syms_bool
+syms_global_iter_next_pdb(SymsDebugInfoPdb *debug_info, SymsGlobalIterPdb *iter, SymsGlobal *gdata_out)
+{
+  syms_bool result = syms_false;
+  pdb_var gdata;
+
+  (void)debug_info;
+
+  if (pdb_global_it_next(&iter->impl, &gdata)) {
+    gdata_out->type_id = syms_typeid_for_pdb(gdata.itype);
+    *(pdb_encoded_location *)gdata_out->encoded_va.impl = gdata.encoded_va;
+    gdata_out->name = syms_string_ref_pdb(&gdata.name);
+    result = syms_true;
+  }
+  return result;
+}
+
+SYMS_INTERNAL syms_bool
+syms_local_data_iter_init_pdb(SymsDebugInfoPdb *debug_info, SymsLocalDataIterPdb *iter, SymsMod *mod)
+{ (void)debug_info;
+  pdb_mod *mod_impl = (pdb_mod *)mod->impl;
+  return pdb_sym_it_init(&iter->impl, mod_impl);
+}
+
+SYMS_INTERNAL syms_bool
+syms_local_data_iter_next_pdb(SymsDebugInfoPdb *debug_info, SymsLocalDataIterPdb *iter, SymsLocalData *ldata_out)
+{
+  syms_bool result = syms_false;
+  pdb_cv_sym_type cv_type;
+  pdb_stream cv_data;
+  for (;!result;) {
+    if (!pdb_sym_it_next(&iter->impl, &cv_type, &cv_data)) {
+      break;
+    }
+    switch (cv_type) {
+    case PDB_CV_SYM_LDATA32: {
+      pdb_cv_datasym32 sym;
+      if (pdb_stream_read(&cv_data, &sym, sizeof(sym))) {
+        pdb_pointer name = pdb_pointer_bake_stream_str(&cv_data);
+        pdb_encoded_location encoded_va = pdb_encode_location_for_datasym32(&debug_info->context, &sym);
+        ldata_out->type_id = syms_typeid_for_pdb(sym.itype);
+        ldata_out->name = syms_string_ref_pdb(&name);
+        *(pdb_encoded_location *)ldata_out->encoded_va.impl = encoded_va;
+        result = syms_true;
+      } else {
+        SYMS_ASSERT_FAILURE("bad cvdata");
+      }
+    } break;
+    case PDB_CV_SYM_LDATA16:
+    case PDB_CV_SYM_LDATA32_16t:
+    case PDB_CV_SYM_LDATA32_ST:
+    case PDB_CV_SYM_LDATA_HLSL:
+    case PDB_CV_SYM_LDATA_HLSL32_EX:
+    case PDB_CV_SYM_LDATA_HLSL32: {
+    } break;
+    }
+  }
+  return result;
+}
+
+SYMS_INTERNAL syms_bool
+SymsFile_iter_init_pdb(SymsDebugInfoPdb *debug_info, SymsFileIterPdb *iter)
+{
+  return pdb_file_it_init(&iter->impl, &debug_info->context);
+}
+
+SYMS_INTERNAL syms_bool
+SymsFile_iter_next_pdb(SymsDebugInfoPdb *debug_info, SymsFileIterPdb *iter, SymsStringRef *ref_out)
+{ 
+  syms_bool result = syms_false;
+  pdb_pointer filename;
+
+  (void)debug_info;
+
+  if (pdb_file_it_next(&iter->impl, &filename)) {
+    *ref_out = syms_string_ref_pdb(&filename);
+    result = syms_true;
+  }
+  return result;
+}
+
+SYMS_INTERNAL syms_bool
+syms_range_iter_init_pdb(SymsDebugInfoPdb *debug_info, SymsRangeIterPdb *iter, SymsAddr lo, SymsAddr hi, pdb_pointer gaps)
+{
+  iter->context = &debug_info->context;
+  iter->lo = lo;
+  iter->hi = hi;
+  iter->gaps = gaps;
+  iter->read_offset = 0; 
+  iter->last_range_emitted = syms_false;
+  return syms_true;
+}
+
+SYMS_INTERNAL syms_bool
+syms_range_iter_next_pdb(SymsRangeIterPdb *iter, SymsAddr *range_lo, SymsAddr *range_hi)
+{
+  pdb_cv_lvar_addr_gap_t gap;
+  syms_bool is_read = syms_false;
+  if (pdb_pointer_read(iter->context, &iter->gaps, iter->read_offset, &gap, sizeof(gap)) == sizeof(gap)) {
+    iter->read_offset += sizeof(gap);
+    *range_lo = iter->lo;
+    *range_hi = iter->lo + gap.off;
+    iter->lo = *range_hi + gap.len;
+    is_read = syms_true;
+  } else {
+    if (!iter->last_range_emitted) {
+      iter->last_range_emitted = syms_true;
+      *range_lo = iter->lo;
+      *range_hi = iter->hi;
+      is_read = syms_true;
+    }
+  }
+  return is_read;
+}
+
+SYMS_INTERNAL syms_bool
+syms_proc_iter_init_pdb(SymsDebugInfoPdb *debug_info, SymsProcIterPdb *iter, SymsMod *mod)
+{ 
+  pdb_mod *mod_impl = (pdb_mod *)mod->impl;
+  
+  (void)debug_info;
+
+  return pdb_proc_it_init(&iter->impl, mod_impl);
+}
+
+SYMS_INTERNAL syms_bool
+syms_proc_iter_next_pdb(SymsDebugInfoPdb *debug_info, SymsProcIterPdb *iter, SymsProc *proc_out)
+{
+  syms_bool result = syms_false;
+  pdb_proc proc;
+  if (pdb_proc_it_next(&iter->impl, &proc)) {
+    result = syms_proc_from_pdb_proc(debug_info, &proc, proc_out);
+  }
+  return result;
+}
+
+SYMS_INTERNAL syms_bool
+syms_arg_iter_init_pdb(SymsDebugInfoPdb *debug_info, SymsArgIterPdb *iter, SymsTypeID *id)
+{
+  return pdb_arg_it_init(&iter->impl, &debug_info->context, syms_typeid_to_pdb(id));
+}
+
+SYMS_INTERNAL syms_bool
+syms_arg_iter_next_pdb(SymsDebugInfoPdb *debug_info, SymsArgIterPdb *iter, SymsTypeID *arg_out)
+{
+  syms_bool result = syms_false;
+  pdb_cv_itype itype;
+
+  (void)debug_info;
+
+  if (pdb_arg_it_next(&iter->impl, &itype)) {
+    *arg_out = syms_typeid_for_pdb(itype);
+    result = syms_true;
+  }
+  return result;
+}
+
+SYMS_INTERNAL syms_bool
+syms_local_iter_init_pdb(SymsDebugInfoPdb *debug_info, SymsLocalIterPdb *iter, SymsProc *proc, void *stack, syms_uint stack_size)
+{
+  SymsProcData *impl = (SymsProcData *)proc->impl;
+  syms_bool is_inited = syms_false;
+  if (pdb_local_it_init_(&iter->impl, &debug_info->context, 0, 0, impl->pdb.cvdata)) {
+    iter->scope_stack = (SymsScope *)stack;
+    iter->scope_max = stack_size / sizeof(SymsScope);
+    SYMS_ASSERT(proc->range.type == SYMS_RANGE_PLAIN);
+    {
+      iter->scope_stack[0].inst_lo = proc->range.u.plain.lo;
+      iter->scope_stack[0].inst_hi = proc->range.u.plain.hi;
+      iter->scope_count = 1;
+    }
+    is_inited = syms_true;
+  }
+  return is_inited;
+}
+
+SYMS_INTERNAL syms_bool
+syms_local_iter_next_pdb(SymsDebugInfoPdb *debug_info, SymsLocalIterPdb *iter, SymsAddr rebase, SymsLocalExport *export_out)
+{  (void)rebase;
+  syms_bool result = syms_false;
+  pdb_local_export exported_data;
+
+  (void)debug_info;
+
+  if (pdb_local_it_next(&iter->impl, &exported_data)) {
+    switch (exported_data.type) {
+    case PDB_LOCAL_EXPORT_VAR: {
+      SymsVar *var_out = &export_out->u.var;
+      pdb_var *var = &exported_data.u.var;
+
+      export_out->type = SYMS_LOCAL_EXPORT_VAR;
+      var_out->type_id = syms_typeid_for_pdb(var->itype);
+      var_out->flags = 0;
+      if (var->flags & PDB_CV_LOCALSYM_FLAG_PARAM) {
+        var_out->flags |= SYMS_VAR_FLAG_PARAM;
+      }
+      if (var->flags & PDB_CV_LOCALSYM_FLAG_COPMGEN) {
+        var_out->flags |= SYMS_VAR_FLAG_COMPILER_GEN;
+      }
+      if (var->flags & PDB_CV_LOCALSYM_FLAG_ALIASED) {
+        var_out->flags |= SYMS_VAR_FLAG_ALIASED;
+      }
+      if (var->flags & PDB_CV_LOCALSYM_FLAG_RETVAL) {
+        var_out->flags |= SYMS_VAR_FLAG_RETVAL;
+      }
+      if (var->flags & PDB_CV_LOCALSYM_FLAG_OPTOUT) {
+        var_out->flags |= SYMS_VAR_FLAG_OPT_OUT;
+      }
+      if (var->flags & PDB_CV_LOCALSYM_FLAG_STATIC) {
+        var_out->flags |= SYMS_VAR_FLAG_STATIC;
+      }
+      if (var->flags & PDB_CV_LOCALSYM_FLAG_GLOBAL) {
+        var_out->flags |= SYMS_VAR_FLAG_GLOBAL;
+      }
+      var_out->name_ref = syms_string_ref_pdb(&var->name);
+      *(pdb_encoded_location *)var_out->encoded_va.impl = var->encoded_va;
+      {
+        SymsRangePdb *impl = (SymsRangePdb *)var_out->range.u.impl;
+        var_out->range.type = SYMS_RANGE_IMPL;
+        SYMS_ASSERT(iter->scope_count > 0);
+        impl->lo = iter->scope_stack[iter->scope_count - 1].inst_lo;
+        impl->hi = iter->scope_stack[iter->scope_count - 1].inst_hi;
+        impl->gaps = var->gaps;
+      }
+      result = syms_true;
+    } break;
+    case PDB_LOCAL_EXPORT_SCOPE: {
+      if (iter->scope_count >= iter->scope_max) {
+        SYMS_ASSERT_FAILURE_PARANOID("out of memory for scope stack");
+        result = syms_false;
+        break;
+      }
+      export_out->type = SYMS_LOCAL_EXPORT_SCOPE;
+      export_out->u.scope.inst_lo = exported_data.u.scope.inst_lo;
+      export_out->u.scope.inst_hi = exported_data.u.scope.inst_hi;
+      result = syms_true;
+      iter->scope_stack[iter->scope_count++] = export_out->u.scope;
+    } break;
+    case PDB_LOCAL_EXPORT_SCOPE_END: {
+      if (iter->scope_count > 0) {
+        iter->scope_count -= 1;
+        export_out->type = SYMS_LOCAL_EXPORT_SCOPE_END;
+        export_out->u.scope = iter->scope_stack[iter->scope_count - 1];
+        result = syms_true;
+      }
+    } break;
+    }
+  }
+
+  return result;
+} 
+SYMS_INTERNAL syms_bool
+syms_inline_iter_init_pdb(SymsDebugInfoPdb *debug_info, SymsInlineIterPdb *iter, SymsProc *proc, SymsAddr pc)
+{
+  pdb_isec sec;
+  pdb_isec_umm sec_off;
+  if (pdb_build_sec_off(&debug_info->context, pc, &sec, &sec_off)) {
+    SymsProcData *impl = (SymsProcData *)proc->impl;
+    if (pdb_inline_it_init(&iter->impl, &debug_info->context, impl->pdb.cvdata, sec_off)) {
+      return syms_true;
+    }
+  }
+  return syms_false;
+}
+
+SYMS_INTERNAL syms_bool
+syms_inline_iter_next_pdb(SymsDebugInfoPdb *debug_info, SymsInlineIterPdb *iter, SymsInlineSite *site_out)
+{ 
+  syms_bool is_site_found = syms_false;
+  pdb_inline_site site;
+
+  (void)debug_info;
+
+  if (pdb_inline_it_next(&iter->impl, &site)) {
+    SymsAddr site_base = 0;
+    pdb_build_va(&debug_info->context, site.sec, site.sec_off, &site_base);
+
+    site_out->name     = syms_string_ref_pdb(&site.name);
+    site_out->type_id  = syms_typeid_for_pdb(site.itype);
+    site_out->range_lo = site_base;
+    site_out->range_hi = site_base + site.size;
+
+    // pdb does not export location where procedure was inlined.
+    // it can be deduced by mapping to source address of first instruction.
+    site_out->call_file.name        = syms_string_ref_str(syms_string_init(0,0));
+    site_out->call_file.chksum_type = SYMS_CHECKSUM_NULL;
+    site_out->call_ln = 0;
+
+    site_out->decl_file.name        = syms_string_ref_pdb(&site.fi.path);
+    site_out->decl_file.chksum_type = SYMS_CHECKSUM_NULL;
+    site_out->decl_ln = site.ln_at_pc;
+
+    site_out->sort_index = site_out->range_lo; 
+
+    is_site_found = syms_true;
+  }
+
+  return is_site_found;
+}
+
+SYMS_INTERNAL syms_bool
+syms_const_iter_init_pdb(SymsDebugInfoPdb *debug_info, SymsConstIterPdb *iter)
+{
+  return pdb_const_it_init(&iter->impl, &debug_info->context);
+}
+
+SYMS_INTERNAL syms_bool
+syms_const_iter_next_pdb(SymsDebugInfoPdb *debug_info, SymsConstIterPdb *iter, SymsConst *const_out)
+{
+  syms_bool result = syms_false;
+  pdb_const_value value;
+
+  (void)debug_info;
+
+  if (pdb_const_it_next(&iter->impl, &value)) {
+    syms_const_convert_from_pdb(&value, const_out);
+    result = syms_true;
+  }
+  return result;
+}
+
+SYMS_INTERNAL syms_bool
+syms_type_iter_init_pdb(SymsDebugInfoPdb *debug_info, SymsTypeIterPdb *iter)
+{
+  return pdb_type_it_init(&iter->impl, &debug_info->context);
+}
+
+SYMS_INTERNAL syms_bool
+syms_type_iter_next_pdb(SymsDebugInfoPdb *debug_info, SymsTypeIterPdb *iter, SymsTypeID *typeid_out)
+{
+  syms_bool result = syms_false;
+  pdb_cv_itype itype;
+  (void)debug_info;
+
+  if (pdb_type_it_next(&iter->impl, &itype)) {
+    *typeid_out = syms_typeid_for_pdb(itype);
+    result = syms_true;
+  }
+  return result;
+}
+
+SYMS_INTERNAL syms_bool
+syms_proc_from_name_pdb(SymsDebugInfoPdb *debug_info, SymsString name, SymsProc *proc_out)
+{
+  syms_bool result = syms_false;
+  pdb_proc proc;
+  if (pdb_proc_from_name(&debug_info->context, name.data, syms_trunc_u32(name.len), &proc)) {
+    result = syms_proc_from_pdb_proc(debug_info, &proc, proc_out);
+  }
+  return result;
+}
+
+SYMS_INTERNAL syms_bool
+syms_global_from_name_pdb(SymsDebugInfoPdb *debug_info, SymsString name, SymsGlobal *gvar_out)
+{
+  syms_bool result = syms_false;
+  pdb_var gvar;
+  if (pdb_global_from_name(&debug_info->context, name.data, syms_trunc_u32(name.len), &gvar)) {
+    gvar_out->type_id = syms_typeid_for_pdb(gvar.itype);
+    *(pdb_encoded_location *)gvar_out->encoded_va.impl = gvar.encoded_va;
+	  gvar_out->name = syms_string_ref_pdb(&gvar.name);
+    result = syms_true;
+  }
+  return result;
+}
+
+SYMS_INTERNAL syms_bool
+syms_proc_from_pdb_proc(SymsDebugInfoPdb *debug_info, pdb_proc *proc, SymsProc *proc_out)
+{ (void)debug_info;
+  if (proc_out && proc) {
+    SymsAddr va = 0;
+    pdb_build_va(&debug_info->context, proc->sec, proc->sec_off, &va);
+
+    proc_out->type_id           = syms_typeid_for_pdb(proc->itype);
+    proc_out->len               = proc->size;
+    proc_out->name_ref          = syms_string_ref_pdb(&proc->name);
+    proc_out->va                = va;
+    proc_out->dbg_start_va      = va;
+    proc_out->dbg_end_va        = va + proc_out->len;
+
+    {
+      proc_out->range.type = SYMS_RANGE_PLAIN;
+      proc_out->range.u.plain.lo = va;
+      proc_out->range.u.plain.hi = va + proc->size;
+    }
+
+    {
+      SymsProcData *impl = (SymsProcData *)proc_out->impl;
+      impl->pdb.cvdata   = proc->cvdata;
+    }
+
+    if (proc_out->len == 0) {
+      proc_out->len = 1;
+    }
+  }
+  return syms_true;
+}
+
+SYMS_INTERNAL syms_bool
+syms_proc_from_va_public_symbols_pdb(SymsDebugInfoPdb *debug_info, SymsAddr va, SymsProc *proc_out)
+{
+  syms_bool result = syms_false;
+  pdb_pointer name;
+  if (pdb_find_nearest_sym(&debug_info->context, va, &name)) {
+    pdb_sc sc;
+    if (pdb_find_nearest_sc(&debug_info->context, va, &sc)) {
+      SymsAddr start_va;
+
+      syms_memset(proc_out, 0, sizeof(*proc_out));
+      if (pdb_build_va(&debug_info->context, sc.sec, sc.sec_off, &start_va)) {
+        proc_out->va = start_va;
+        proc_out->len = sc.size;
+        proc_out->dbg_start_va = proc_out->va;
+        proc_out->dbg_end_va = proc_out->va + proc_out->len;
+        proc_out->name_ref = syms_string_ref_pdb(&name);
+        result = syms_true;
+      }
+    }
+  }
+  return result;
+}
+
+SYMS_INTERNAL syms_bool
+syms_type_from_name_pdb(SymsDebugInfoPdb *debug_info, SymsString name, SymsType *type_out)
+{
+  syms_bool result = syms_false;
+  pdb_type type;
+  if (pdb_type_from_name(&debug_info->context, name.data, syms_trunc_u32(name.len), &type)) {
+    result = syms_pdb_type_to_syms_type(&debug_info->context, &type, type_out);
+  }
+  return result;
+}
+
+SYMS_INTERNAL syms_bool
+syms_const_from_name_pdb(SymsDebugInfoPdb *debug_info, SymsString name, SymsConst *const_out)
+{
+  syms_bool result = syms_false;
+  pdb_const_value pdb_const;
+  if (pdb_const_from_name(&debug_info->context, name.data, syms_trunc_u32(name.len), &pdb_const)) {
+    syms_const_convert_from_pdb(&pdb_const, const_out);
+    result = syms_true;
+  }
+  return result;
+}
+
+SYMS_INTERNAL syms_bool
+syms_infer_type_pdb(SymsDebugInfoPdb *debug_info, SymsTypeID type_id, SymsType *type_out)
+{
+  pdb_cv_itype itype = syms_typeid_to_pdb(&type_id);
+  pdb_type type;
+  syms_bool is_inferred = syms_false;
+  if (pdb_infer_itype(&debug_info->context, itype, &type)) {
+    is_inferred = syms_pdb_type_to_syms_type(&debug_info->context, &type, type_out);
+  }
+  return is_inferred;
+}
+
+SYMS_INTERNAL syms_bool
+syms_trampoline_from_ip_pdb(SymsDebugInfoPdb *debug_info, SymsAddr rva, SymsAddr *rva_out)
+{
+  syms_bool result = syms_false;
+  pdb_isec sec;
+  pdb_isec_umm sec_off;
+  if (pdb_build_sec_off(&debug_info->context, rva, &sec, &sec_off)) {
+    pdb_isec dst_sec;
+    pdb_isec_umm dst_sec_off;
+    if (pdb_trampoline_from_ip(&debug_info->context, sec, sec_off, &dst_sec, &dst_sec_off)) {
+      result = pdb_build_va(&debug_info->context, dst_sec, dst_sec_off, rva_out);
+    }
+  }
+  return result;
+}
+
+SYMS_INTERNAL SymsErrorCode
+syms_decode_location_pdb(
+             SymsAddr orig_rebase,
+             SymsAddr rebase,
+             pdb_encoded_location *impl, 
+             SymsRegread *regread_context,
+             SymsMemread *memread_context,
+             SymsLocation *loc_out)
+{
+  SymsErrorCode result = SYMS_ERR_INVALID_CODE_PATH;
+  pdb_location loc;
+
+  (void)memread_context; // TODO(nick): Consider if we need memread here at all
+
+  if (pdb_decode_location(impl, orig_rebase, rebase, 0, 0, regread_context, syms_regread_pdb, &loc)) {
+    switch (loc.type) {
+    case PDB_LOCATION_VA: {
+      loc_out->kind = SYMS_LOCATION_VA;
+      loc_out->u.va = loc.u.va;
+      memread_context->result = SYMS_ERR_OK;
+      result = SYMS_ERR_OK;
+    } break;
+    case PDB_LOCATION_IMPLICIT: {
+      void *dst, *src;
+      syms_uint copy_size;
+
+      loc_out->kind = SYMS_LOCATION_IMPLICIT;
+      loc_out->u.implicit.len = loc.u.implicit.len;
+
+      dst = &loc_out->u.implicit.data[0];
+      src = &loc.u.implicit.data[0];
+      copy_size = sizeof(loc_out->u.implicit.data);
+      syms_memcpy(dst, src, copy_size);
+
+      memread_context->result = SYMS_ERR_OK;
+      result = SYMS_ERR_OK;
+    } break;
+    case PDB_LOCATION_NULL: {
+      loc_out->kind = SYMS_LOCATION_NULL;
+
+      memread_context->result = SYMS_ERR_OK;
+      result = SYMS_ERR_OK;
+    } break;
+    }
+  }
+
+  return result;
+}
+
+SYMS_INTERNAL void
+syms_const_convert_from_pdb(pdb_const_value *pdb_const, SymsConst *const_out)
+{
+  u8 value_size = PDB_BASIC_TYPE_SIZE_MASK(pdb_const->value.itype);
+
+  if (value_size <= sizeof(const_out->value)) {
+    const_out->type_id = syms_typeid_for_pdb(pdb_const->itype);
+    const_out->name = syms_string_ref_pdb(&pdb_const->name);
+    const_out->value_len = sizeof(pdb_const->value);
+    syms_memcpy(&const_out->value[0], &pdb_const->value.u.data[0], sizeof(pdb_const->value));
+  } else {
+    const_out->value_len = 0;
+    SYMS_ASSERT_FAILURE("cannot copy pdb constant value - insufficient size of destination storage");
+  }
+}
+
+SYMS_INTERNAL syms_bool
+syms_pdb_type_to_syms_type(pdb_context *pdb, pdb_type *type, SymsType *type_out)
+{
+  type_out->is_fwdref = syms_false;
+  type_out->id = syms_typeid_for_pdb(type->cv_itype);
+  type_out->next_id = syms_typeid_for_pdb(type->next_cv_itype);
+  
+  type_out->modifier = 0; 
+  if (type->attribs & PDB_TYPE_ATTRIB_CONST)      type_out->modifier |= SYMS_TYPE_MDFR_CONST;
+  if (type->attribs & PDB_TYPE_ATTRIB_VOLATILE)   type_out->modifier |= SYMS_TYPE_MDFR_VOLATILE;
+  if (type->attribs & PDB_TYPE_ATTRIB_UNALIGNED)  type_out->modifier |= SYMS_TYPE_MDFR_PACKED;
+  if (type->attribs & PDB_TYPE_ATTRIB_RESTRICTED) type_out->modifier |= SYMS_TYPE_MDFR_RESTRICT; 
+  if (type->attribs & PDB_TYPE_ATTRIB_LREF)       type_out->modifier |= SYMS_TYPE_MDFR_REF;
+  if (type->attribs & PDB_TYPE_ATTRIB_RREF)       type_out->modifier |= SYMS_TYPE_MDFR_RVALUE_REF;
+  if (type->attribs & PDB_TYPE_FWDREF)            type_out->modifier |= SYMS_TYPE_MDFR_FWDREF;
+  
+  switch (type->kind) {
+  default:
+  case PDB_TYPE_NULL: type_out->kind = SYMS_TYPE_NULL; break;
+
+  case PDB_TYPE_CHAR: {
+    type_out->kind = SYMS_TYPE_INT8;
+    type_out->modifier |= SYMS_TYPE_MDFR_CHAR;
+  } break;
+  case PDB_TYPE_UCHAR: {
+    type_out->kind = SYMS_TYPE_UINT8;
+    type_out->modifier |= SYMS_TYPE_MDFR_CHAR;
+  } break;
+  case PDB_TYPE_WCHAR: {
+    type_out->kind = SYMS_TYPE_UINT16;
+    type_out->modifier |= SYMS_TYPE_MDFR_CHAR;
+  } break;
+
+  case PDB_TYPE_VARIADIC: type_out->kind = SYMS_TYPE_VARIADIC; break;
+
+  case PDB_TYPE_UINT8:    type_out->kind = SYMS_TYPE_UINT8; break;
+  case PDB_TYPE_UINT16:   type_out->kind = SYMS_TYPE_UINT16; break;
+  case PDB_TYPE_UINT32:   type_out->kind = SYMS_TYPE_UINT32; break;
+  case PDB_TYPE_UINT64:   type_out->kind = SYMS_TYPE_UINT64; break;
+
+  case PDB_TYPE_INT8:  type_out->kind = SYMS_TYPE_INT8;   break;
+  case PDB_TYPE_INT16: type_out->kind = SYMS_TYPE_INT16;  break;
+  case PDB_TYPE_INT32: type_out->kind = SYMS_TYPE_INT32;  break;
+  case PDB_TYPE_INT64: type_out->kind = SYMS_TYPE_INT64;  break;
+
+  case PDB_TYPE_REAL16:   type_out->kind = SYMS_TYPE_FLOAT16;      break;
+  case PDB_TYPE_REAL32:   type_out->kind = SYMS_TYPE_FLOAT32;      break;
+  case PDB_TYPE_REAL32PP: type_out->kind = SYMS_TYPE_FLOAT32PP;    break;
+  case PDB_TYPE_REAL64:   type_out->kind = SYMS_TYPE_FLOAT64;      break;
+  case PDB_TYPE_REAL80:   type_out->kind = SYMS_TYPE_FLOAT80;      break; 
+  case PDB_TYPE_REAL128:  type_out->kind = SYMS_TYPE_FLOAT128;     break;
+
+  case PDB_TYPE_COMPLEX32:    type_out->kind = SYMS_TYPE_COMPLEX32;   break;
+  case PDB_TYPE_COMPLEX64:    type_out->kind = SYMS_TYPE_COMPLEX64;   break;
+  case PDB_TYPE_COMPLEX80:    type_out->kind = SYMS_TYPE_COMPLEX80;   break;
+  case PDB_TYPE_COMPLEX128:   type_out->kind = SYMS_TYPE_COMPLEX128;  break;
+
+  case PDB_TYPE_BOOL: type_out->kind = SYMS_TYPE_BOOL; break;
+  case PDB_TYPE_VOID: type_out->kind = SYMS_TYPE_VOID; break;
+
+  case PDB_TYPE_ENUM:     type_out->kind = SYMS_TYPE_ENUM;   break;
+  case PDB_TYPE_UNION:    type_out->kind = SYMS_TYPE_UNION;  break;
+  case PDB_TYPE_STRUCT:   type_out->kind = SYMS_TYPE_STRUCT; break;
+  case PDB_TYPE_CLASS:    type_out->kind = SYMS_TYPE_CLASS;  break;
+
+  case PDB_TYPE_PROC: {
+    type_out->kind = SYMS_TYPE_PROC; 
+    type_out->u.proc.arglist_type_id = syms_typeid_for_pdb(type->u.proc.arg_itype);
+    type_out->u.proc.ret_type_id     = syms_typeid_for_pdb(type->u.proc.ret_itype);
+    type_out->u.proc.arg_count       = type->u.proc.arg_count;
+  } break;
+
+  case PDB_TYPE_METHOD: {
+    type_out->kind = SYMS_TYPE_METHOD;
+    type_out->u.method.arg_count = type->u.method.arg_count;
+    type_out->u.method.ret_type_id = syms_typeid_for_pdb(type->u.method.ret_itype);
+    type_out->u.method.class_type_id = syms_typeid_for_pdb(type->u.method.class_itype);
+    type_out->u.method.this_type_id = syms_typeid_for_pdb(type->u.method.this_itype);
+    type_out->u.method.arglist_type_id = syms_typeid_for_pdb(type->u.method.arg_itype);
+  } break;
+
+  case PDB_TYPE_BITFIELD: {
+    type_out->kind = SYMS_TYPE_BITFIELD;
+    type_out->u.bitfield.base_type_id = syms_typeid_for_pdb(type->u.bitfield.base_itype);
+    type_out->u.bitfield.len = type->u.bitfield.len;
+    type_out->u.bitfield.pos = type->u.bitfield.pos;
+  } break;
+
+  case PDB_TYPE_PTR: {
+    type_out->kind = SYMS_TYPE_PTR; 
+    switch (type->u.ptr.mode) {
+    case PDB_CV_PTR_MODE_PTR:    type_out->u.ptr.mode = SYMS_PTR_MODE_NORMAL; break;
+    case PDB_CV_PTR_MODE_LVREF:  type_out->u.ptr.mode = SYMS_PTR_MODE_LVREF;  break;
+    case PDB_CV_PTR_MODE_RVREF:  type_out->u.ptr.mode = SYMS_PTR_MODE_RVREF;  break;
+    case PDB_CV_PTR_MODE_PMEM:   type_out->u.ptr.mode = SYMS_PTR_MODE_MEM;    break;
+    case PDB_CV_PTR_MODE_PMFUNC: type_out->u.ptr.mode = SYMS_PTR_MODE_MFUNC;  break;
+    }
+  } break;
+
+  case PDB_TYPE_ARR: type_out->kind = SYMS_TYPE_ARR; break;
+  }
+
+  type_out->is_fwdref = (type->attribs & PDB_TYPE_ATTRIB_FWDREF) != 0;
+  
+  if (type->kind == PDB_TYPE_ARR) {
+    pdb_type next_type;
+    if (pdb_infer_itype(pdb, type->next_cv_itype, &next_type))
+      if (next_type.size != 0)
+      type_out->size = type->size / next_type.size;
+  } else {
+    type_out->size = type->size;
+  }
+  type_out->name_ref = syms_string_ref_pdb(&type->name);
+
+  // For nested UDT compiler writes a special string that 
+  // we check against.
+  {
+    char name[16];
+    syms_uint name_size = pdb_pointer_read(pdb, &type->name, 0, name, sizeof(name) - 1);
+    name[name_size] = 0;
+    if (syms_strcmp(name, "<anonymous-tag>") == 0) {
+      type_out->modifier |= SYMS_TYPE_MDFR_NESTED;
+    }
+  }
+  
+  // TODO(nick): A bit work has to be done on the PDB API side
+  // in order for us to be able to remove this.
+  *(pdb_type *)type_out->impl = *type;
+
+  {
+    pdb_udt_srcline udt_srcline;
+    if (pdb_find_udt_srcline(pdb, type->cv_itype, &udt_srcline)) {
+      type_out->decl_file = syms_string_ref_pdb(&udt_srcline.file);
+      type_out->decl_ln = udt_srcline.ln;
+    } else {
+      type_out->decl_file = syms_string_ref_null();
+      type_out->decl_ln = 0;
+    }
+  }
+
+  return syms_true;
+}
+
+SYMS_INTERNAL SymsStringRef
+syms_string_ref_pdb(pdb_pointer *pointer)
+{
+  SymsStringRef ref;
+  ref.type = SYMS_STRING_REF_TYPE_PDB;
+  syms_memcpy(&ref.impl[0], pointer, sizeof(*pointer));
+  return ref;
+}
+
+SYMS_INTERNAL syms_bool
+syms_string_cmp_strref_pdb(SymsDebugInfoPdb *debug_info, SymsString a, pdb_pointer b)
+{
+  return pdb_strcmp_pointer(&debug_info->context, a, &b);
+}
+
+SYMS_INTERNAL syms_uint
+syms_read_strref_pdb(SymsDebugInfoPdb *debug_info, pdb_pointer *ref, void *buffer, pdb_uint buffer_max)
+{
+  pdb_uint read_size = pdb_pointer_get_size(ref);
+  if (buffer == NULL) {
+    SYMS_ASSERT(buffer_max == 0);
+    return read_size + 1; // reserve 1 byte for null
+  }
+  if (buffer_max > 0) {
+    read_size = pdb_pointer_read(&debug_info->context, ref, 0, buffer, buffer_max);
+    if (read_size < buffer_max) {
+      ((U8 *)buffer)[read_size++] = '\0';
+    } else {
+      ((U8 *)buffer)[buffer_max - 1] = '\0';
+    }
+  }
+  return (syms_uint)read_size;
+}
+
+SYMS_INTERNAL SymsRegID
+syms_remap_regid_pdb(syms_uint regid)
+{
+  static syms_bool init = syms_true;
+  static SymsRegID regs_table[PDB_CV_REG_MAX];
+
+  if (init) {
+    regs_table[PDB_CV_X86_AL]       = SYMS_REG_X86_al;
+    regs_table[PDB_CV_X86_CL]       = SYMS_REG_X86_cl;
+    regs_table[PDB_CV_X86_DL]       = SYMS_REG_X86_dl;
+    regs_table[PDB_CV_X86_BL]       = SYMS_REG_X86_bl;
+    regs_table[PDB_CV_X86_AH]       = SYMS_REG_X86_ah;
+    regs_table[PDB_CV_X86_CH]       = SYMS_REG_X86_ch;
+    regs_table[PDB_CV_X86_DH]       = SYMS_REG_X86_dh;
+    regs_table[PDB_CV_X86_BH]       = SYMS_REG_X86_bh;
+    regs_table[PDB_CV_X86_AX]       = SYMS_REG_X86_ax;
+    regs_table[PDB_CV_X86_CX]       = SYMS_REG_X86_cx;
+    regs_table[PDB_CV_X86_DX]       = SYMS_REG_X86_dx;
+    regs_table[PDB_CV_X86_BX]       = SYMS_REG_X86_bx;
+    regs_table[PDB_CV_X86_SP]       = SYMS_REG_X86_sp;
+    regs_table[PDB_CV_X86_BP]       = SYMS_REG_X86_bp;
+    regs_table[PDB_CV_X86_SI]       = SYMS_REG_X86_si;
+    regs_table[PDB_CV_X86_DI]       = SYMS_REG_X86_di;
+    regs_table[PDB_CV_X86_EAX]      = SYMS_REG_X86_eax;
+    regs_table[PDB_CV_X86_ECX]      = SYMS_REG_X86_ecx;
+    regs_table[PDB_CV_X86_EDX]      = SYMS_REG_X86_edx;
+    regs_table[PDB_CV_X86_EBX]      = SYMS_REG_X86_ebx;
+    regs_table[PDB_CV_X86_ESP]      = SYMS_REG_X86_esp;
+    regs_table[PDB_CV_X86_EBP]      = SYMS_REG_X86_ebp;
+    regs_table[PDB_CV_X86_ESI]      = SYMS_REG_X86_esi;
+    regs_table[PDB_CV_X86_EDI]      = SYMS_REG_X86_edi;
+    regs_table[PDB_CV_X86_ES]       = SYMS_REG_X86_es;
+    regs_table[PDB_CV_X86_CS]       = SYMS_REG_X86_cs;
+    regs_table[PDB_CV_X86_SS]       = SYMS_REG_X86_ss;
+    regs_table[PDB_CV_X86_DS]       = SYMS_REG_X86_ds;
+    regs_table[PDB_CV_X86_FS]       = SYMS_REG_X86_fs;
+    regs_table[PDB_CV_X86_GS]       = SYMS_REG_X86_gs;
+    regs_table[PDB_CV_X86_IP]       = SYMS_REG_X86_ip;
+    regs_table[PDB_CV_X86_EIP]      = SYMS_REG_X86_eip;
+    regs_table[PDB_CV_X64_EFLAGS]   = SYMS_REG_X86_eflags;
+
+    regs_table[PDB_CV_X64_RAX]    = SYMS_REG_X64_rax;
+    regs_table[PDB_CV_X64_RBX]    = SYMS_REG_X64_rbx;
+    regs_table[PDB_CV_X64_RCX]    = SYMS_REG_X64_rcx;
+    regs_table[PDB_CV_X64_RDX]    = SYMS_REG_X64_rdx;
+    regs_table[PDB_CV_X64_RSI]    = SYMS_REG_X64_rsi;
+    regs_table[PDB_CV_X64_RDI]    = SYMS_REG_X64_rdi;
+    regs_table[PDB_CV_X64_RBP]    = SYMS_REG_X64_rbp;
+    regs_table[PDB_CV_X64_RSP]    = SYMS_REG_X64_rsp;
+    regs_table[PDB_CV_X64_R8]     = SYMS_REG_X64_r8;
+    regs_table[PDB_CV_X64_R9]     = SYMS_REG_X64_r9;
+    regs_table[PDB_CV_X64_R10]    = SYMS_REG_X64_r10;
+    regs_table[PDB_CV_X64_R11]    = SYMS_REG_X64_r11;
+    regs_table[PDB_CV_X64_R12]    = SYMS_REG_X64_r12;
+    regs_table[PDB_CV_X64_R13]    = SYMS_REG_X64_r13;
+    regs_table[PDB_CV_X64_R14]    = SYMS_REG_X64_r14;
+    regs_table[PDB_CV_X64_R15]    = SYMS_REG_X64_r15;
+    regs_table[PDB_CV_X64_R8B]    = SYMS_REG_X64_r8b;
+    regs_table[PDB_CV_X64_R9B]    = SYMS_REG_X64_r9b;
+    regs_table[PDB_CV_X64_R10B]   = SYMS_REG_X64_r10b;
+    regs_table[PDB_CV_X64_R11B]   = SYMS_REG_X64_r11b;
+    regs_table[PDB_CV_X64_R12B]   = SYMS_REG_X64_r12b;
+    regs_table[PDB_CV_X64_R13B]   = SYMS_REG_X64_r13b;
+    regs_table[PDB_CV_X64_R14B]   = SYMS_REG_X64_r14b;
+    regs_table[PDB_CV_X64_R15B]   = SYMS_REG_X64_r15b;
+    regs_table[PDB_CV_X64_R8W]    = SYMS_REG_X64_r8w;
+    regs_table[PDB_CV_X64_R9W]    = SYMS_REG_X64_r9w;
+    regs_table[PDB_CV_X64_R10W]   = SYMS_REG_X64_r10w;
+    regs_table[PDB_CV_X64_R11W]   = SYMS_REG_X64_r11w;
+    regs_table[PDB_CV_X64_R12W]   = SYMS_REG_X64_r12w;
+    regs_table[PDB_CV_X64_R13W]   = SYMS_REG_X64_r13w;
+    regs_table[PDB_CV_X64_R14W]   = SYMS_REG_X64_r14w;
+    regs_table[PDB_CV_X64_R15W]   = SYMS_REG_X64_r15w;
+    regs_table[PDB_CV_X64_R8D]    = SYMS_REG_X64_r8d;
+    regs_table[PDB_CV_X64_R9D]    = SYMS_REG_X64_r9d;
+    regs_table[PDB_CV_X64_R10D]   = SYMS_REG_X64_r10d;
+    regs_table[PDB_CV_X64_R11D]   = SYMS_REG_X64_r11d;
+    regs_table[PDB_CV_X64_R12D]   = SYMS_REG_X64_r12d;
+    regs_table[PDB_CV_X64_R13D]   = SYMS_REG_X64_r13d;
+    regs_table[PDB_CV_X64_R14D]   = SYMS_REG_X64_r14d;
+    regs_table[PDB_CV_X64_R15D]   = SYMS_REG_X64_r15d;
+    regs_table[PDB_CV_X64_SIL]    = SYMS_REG_X64_sil;
+    regs_table[PDB_CV_X64_DIL]    = SYMS_REG_X64_dil;
+    regs_table[PDB_CV_X64_BPL]    = SYMS_REG_X64_bpl;
+    regs_table[PDB_CV_X64_SPL]    = SYMS_REG_X64_spl;
+    regs_table[PDB_CV_X64_YMM0]   = SYMS_REG_X64_ymm0;
+    regs_table[PDB_CV_X64_YMM1]   = SYMS_REG_X64_ymm1;
+    regs_table[PDB_CV_X64_YMM2]   = SYMS_REG_X64_ymm2;
+    regs_table[PDB_CV_X64_YMM3]   = SYMS_REG_X64_ymm3;
+    regs_table[PDB_CV_X64_YMM4]   = SYMS_REG_X64_ymm4;
+    regs_table[PDB_CV_X64_YMM5]   = SYMS_REG_X64_ymm5;
+    regs_table[PDB_CV_X64_YMM6]   = SYMS_REG_X64_ymm6;
+    regs_table[PDB_CV_X64_YMM7]   = SYMS_REG_X64_ymm7;
+    regs_table[PDB_CV_X64_YMM8]   = SYMS_REG_X64_ymm8;
+    regs_table[PDB_CV_X64_YMM9]   = SYMS_REG_X64_ymm9;
+    regs_table[PDB_CV_X64_YMM10]  = SYMS_REG_X64_ymm10;
+    regs_table[PDB_CV_X64_YMM11]  = SYMS_REG_X64_ymm11;
+    regs_table[PDB_CV_X64_YMM12]  = SYMS_REG_X64_ymm12;
+    regs_table[PDB_CV_X64_YMM13]  = SYMS_REG_X64_ymm13;
+    regs_table[PDB_CV_X64_YMM14]  = SYMS_REG_X64_ymm14;
+    regs_table[PDB_CV_X64_YMM15]  = SYMS_REG_X64_ymm15;
+    regs_table[PDB_CV_X64_XMM0]   = SYMS_REG_X64_xmm0;
+    regs_table[PDB_CV_X64_XMM1]   = SYMS_REG_X64_xmm1;
+    regs_table[PDB_CV_X64_XMM2]   = SYMS_REG_X64_xmm2;
+    regs_table[PDB_CV_X64_XMM3]   = SYMS_REG_X64_xmm3;
+    regs_table[PDB_CV_X64_XMM4]   = SYMS_REG_X64_xmm4;
+    regs_table[PDB_CV_X64_XMM5]   = SYMS_REG_X64_xmm5;
+    regs_table[PDB_CV_X64_XMM6]   = SYMS_REG_X64_xmm6;
+    regs_table[PDB_CV_X64_XMM7]   = SYMS_REG_X64_xmm7;
+    regs_table[PDB_CV_X64_XMM8]   = SYMS_REG_X64_xmm8;
+    regs_table[PDB_CV_X64_XMM9]   = SYMS_REG_X64_xmm9;
+    regs_table[PDB_CV_X64_XMM10]  = SYMS_REG_X64_xmm10;
+    regs_table[PDB_CV_X64_XMM11]  = SYMS_REG_X64_xmm11;
+    regs_table[PDB_CV_X64_XMM12]  = SYMS_REG_X64_xmm12;
+    regs_table[PDB_CV_X64_XMM13]  = SYMS_REG_X64_xmm13;
+    regs_table[PDB_CV_X64_XMM14]  = SYMS_REG_X64_xmm14;
+    regs_table[PDB_CV_X64_XMM15]  = SYMS_REG_X64_xmm15;
+    regs_table[PDB_CV_X64_DR0]    = SYMS_REG_X64_dr0;
+    regs_table[PDB_CV_X64_DR1]    = SYMS_REG_X64_dr1;
+    regs_table[PDB_CV_X64_DR2]    = SYMS_REG_X64_dr2;
+    regs_table[PDB_CV_X64_DR3]    = SYMS_REG_X64_dr3;
+    regs_table[PDB_CV_X64_DR4]    = SYMS_REG_X64_dr4;
+    regs_table[PDB_CV_X64_DR5]    = SYMS_REG_X64_dr5;
+    regs_table[PDB_CV_X64_DR6]    = SYMS_REG_X64_dr6;
+    regs_table[PDB_CV_X64_DR7]    = SYMS_REG_X64_dr7;
+    regs_table[PDB_CV_X64_AL]     = SYMS_REG_X64_al;
+    regs_table[PDB_CV_X64_CL]     = SYMS_REG_X64_cl;
+    regs_table[PDB_CV_X64_DL]     = SYMS_REG_X64_dl;
+    regs_table[PDB_CV_X64_BL]     = SYMS_REG_X64_bl;
+    regs_table[PDB_CV_X64_AH]     = SYMS_REG_X64_ah;
+    regs_table[PDB_CV_X64_CH]     = SYMS_REG_X64_ch;
+    regs_table[PDB_CV_X64_DH]     = SYMS_REG_X64_dh;
+    regs_table[PDB_CV_X64_BH]     = SYMS_REG_X64_bh;
+    regs_table[PDB_CV_X64_AX]     = SYMS_REG_X64_ax;
+    regs_table[PDB_CV_X64_CX]     = SYMS_REG_X64_cx;
+    regs_table[PDB_CV_X64_DX]     = SYMS_REG_X64_dx;
+    regs_table[PDB_CV_X64_BX]     = SYMS_REG_X64_bx;
+    regs_table[PDB_CV_X64_SP]     = SYMS_REG_X64_sp;
+    regs_table[PDB_CV_X64_BP]     = SYMS_REG_X64_bp;
+    regs_table[PDB_CV_X64_SI]     = SYMS_REG_X64_si;
+    regs_table[PDB_CV_X64_DI]     = SYMS_REG_X64_di;
+    regs_table[PDB_CV_X64_EAX]    = SYMS_REG_X64_eax;
+    regs_table[PDB_CV_X64_ECX]    = SYMS_REG_X64_ecx;
+    regs_table[PDB_CV_X64_EDX]    = SYMS_REG_X64_edx;
+    regs_table[PDB_CV_X64_EBX]    = SYMS_REG_X64_ebx;
+    regs_table[PDB_CV_X64_ESP]    = SYMS_REG_X64_esp;
+    regs_table[PDB_CV_X64_EBP]    = SYMS_REG_X64_ebp;
+    regs_table[PDB_CV_X64_ESI]    = SYMS_REG_X64_esi;
+    regs_table[PDB_CV_X64_EDI]    = SYMS_REG_X64_edi;
+    regs_table[PDB_CV_X64_ES]     = SYMS_REG_X64_es;
+    regs_table[PDB_CV_X64_CS]     = SYMS_REG_X64_cs;
+    regs_table[PDB_CV_X64_SS]     = SYMS_REG_X64_ss;
+    regs_table[PDB_CV_X64_DS]     = SYMS_REG_X64_ds;
+    regs_table[PDB_CV_X64_FS]     = SYMS_REG_X64_fs;
+    regs_table[PDB_CV_X64_GS]     = SYMS_REG_X64_gs;
+    regs_table[PDB_CV_X64_RIP]    = SYMS_REG_X64_rip;
+    regs_table[PDB_CV_X64_EFLAGS] = SYMS_REG_X64_rflags;
+    regs_table[PDB_CV_X64_ST0]    = SYMS_REG_X64_st0;
+    regs_table[PDB_CV_X64_ST1]    = SYMS_REG_X64_st1;
+    regs_table[PDB_CV_X64_ST2]    = SYMS_REG_X64_st2;
+    regs_table[PDB_CV_X64_ST3]    = SYMS_REG_X64_st3;
+    regs_table[PDB_CV_X64_ST4]    = SYMS_REG_X64_st4;
+    regs_table[PDB_CV_X64_ST5]    = SYMS_REG_X64_st5;
+    regs_table[PDB_CV_X64_ST6]    = SYMS_REG_X64_st6;
+    regs_table[PDB_CV_X64_ST7]    = SYMS_REG_X64_st7;
+
+    init = syms_false;
+  }
+
+  return (regid < SYMS_ARRAY_SIZE(regs_table)) ? regs_table[regid] : SYMS_REG_null;
+}
+
+SYMS_INTERNAL
+PDB_REGREAD_SIG(syms_regread_pdb)
+{
+  SymsRegread *regread = (SymsRegread *)context;
+  SymsRegID regid = syms_remap_regid_pdb(reg_index);
+  syms_uint num_read = regread->callback(regread->context, regread->arch, regid, read_buffer, read_buffer_max);
+  return syms_trunc_u32(num_read);
+#if 0
+  U8 gpr_len = 0; 
+  U8 ptr_len = 0;
+  U32 num_read = 0;
+  switch (syms_regs_get_arch(regs)) {
+  case SYMS_ARCH_X86: gpr_len = 4; ptr_len = 4; break;
+  case SYMS_ARCH_X64: gpr_len = 8; ptr_len = 8; break;
+  default: {
+    SYMS_ASSERT_FAILURE("implement");
+  } break;
+  }
+    if (regid != SYMS_REG_null) {
+      const SymsRegDesc *desc;
+
+      num_read = (U32)syms_regs_get_value(regs, regid, read_buffer, (U32)read_buffer_max);
+
+      desc = syms_regs_get_regdesc(regs, regid);
+      if (desc) {
+        u64 value;
+
+        SYMS_ASSERT(desc->bitpos_le + desc->bitcount_le <= 128);
+        SYMS_ASSERT(desc->bitpos_le >= 0 && desc->bitcount_le <= 64);
+
+        value = *(U64 *)read_buffer;
+        value = (value >> desc->bitpos_le) & (SYMS_UINT64_MAX >> desc->bitpos_le);
+
+        *(U64 *)read_buffer = value;
+      }
+    } else {
+      SYMS_ASSERT_FAILURE("cannot access register with invalid id");
+    }
+#endif
+}
+
+SYMS_INTERNAL
+PDB_MEMREAD_SIG(syms_memread_pdb)
+{
+  SymsMemread *memread = (SymsMemread *)context;
+  SymsErrorCode memread_result = memread->callback(memread->context, va, read_buffer, read_buffer_max);
+  return SYMS_RESULT_OK(memread_result) ? syms_true : syms_false;
+}
+
+SYMS_INTERNAL SymsModID
+syms_infer_global_data_module_pdb(struct SymsInstance *instance, SymsDebugInfoPdb *debug_info, SymsGlobal *gdata)
+{
+  SymsModID mod_id = SYMS_INVALID_MOD_ID;
+  SymsLocation loc;
+  SymsErrorCode err = syms_decode_location(instance, &gdata->encoded_va, 0, 0, 0, 0, &loc);
+  if (SYMS_RESULT_OK(err)) {
+    if (loc.kind == SYMS_LOCATION_VA) {
+      SymsAddr rebase = syms_get_rebase(instance); 
+      SymsAddr rva = loc.u.va - rebase;
+      pdb_sc sc;
+      if (pdb_find_nearest_sc(&debug_info->context, rva, &sc)) {
+        syms_uint i;
+        for (i = 0; i < instance->mods_num; ++i) {
+          pdb_mod *mod = (pdb_mod *)instance->mods[i].header.impl;
+          if (mod->id == sc.imod) {
+            mod_id = (SymsModID)i;
+            break;
+          }
+        }
+      }
+    }
+  }
+  return mod_id;
+}
+
diff --git a/Engine/Source/Runtime/SymsLib/src/pdb/syms_pdb_api.h b/Engine/Source/Runtime/SymsLib/src/pdb/syms_pdb_api.h
new file mode 100644
--- /dev/null
+++ b/Engine/Source/Runtime/SymsLib/src/pdb/syms_pdb_api.h
@@ -0,0 +1,154 @@
+// Copyright Epic Games, Inc. All Rights Reserved.
+#ifndef SYMS_SYMS_API_INCLUDE_H
+#define SYMS_SYMS_API_INCLUDE_H
+
+/******************************************************************************
+ * File   : syms_pdb.h                                                        *
+ * Author : Nikita Smith                                                      *
+ * Created: 2020/06/20                                                        *
+ * Purpose: Declarations for pdb.h interface                                  *
+ ******************************************************************************/
+
+typedef struct SymsDebugInfoPdb 
+{
+  pdb_context context;
+} SymsDebugInfoPdb;
+
+typedef struct SymsModIterPdb 
+{
+  pdb_mod_it impl;
+} SymsModIterPdb;
+
+typedef struct SymsLineIterPdb 
+{
+  pdb_line_it impl;
+} SymsLineIterPdb;
+
+typedef struct SymsMemberIterPdb 
+{
+  pdb_member_it impl;
+  pdb_pointer methodlist_name;
+  pdb_uint methodlist_index;
+  pdb_uint methodlist_count;
+  pdb_pointer methodlist_block;
+} SymsMemberIterPdb;
+
+typedef struct SymsGlobalIterPdb 
+{
+  pdb_global_it impl;
+} SymsGlobalIterPdb;
+
+typedef struct SymsLocalDataIterPdb 
+{
+  pdb_sym_it impl;
+} SymsLocalDataIterPdb;
+
+typedef struct SymsFileIterPdb 
+{
+  pdb_file_it impl;
+} SymsFileIterPdb;
+
+typedef struct SymsProcIterPdb 
+{
+  pdb_proc_it impl;
+} SymsProcIterPdb;
+
+typedef struct SymsArgIterPdb 
+{
+  pdb_arg_it impl;
+} SymsArgIterPdb;
+
+typedef struct SymsLocalIterPdb 
+{
+  pdb_local_it impl;
+  SymsScope *scope_stack;
+  syms_uint scope_count;
+  syms_uint scope_max;
+} SymsLocalIterPdb;
+
+typedef struct SymsInlineIterPdb 
+{
+  pdb_inline_it impl;
+} SymsInlineIterPdb;
+
+typedef struct SymsConstIterPdb 
+{
+  pdb_const_it impl;
+} SymsConstIterPdb;
+
+typedef struct SymsTypeIterPdb 
+{
+  pdb_type_it impl;
+} SymsTypeIterPdb;
+
+typedef struct SymsRangePdb
+{
+  SymsAddr lo;
+  SymsAddr hi;
+  pdb_pointer gaps;
+} SymsRangePdb;
+
+typedef struct SymsRangeIterPdb
+{
+  pdb_context *context;
+  pdb_pointer gaps;
+  SymsAddr lo;
+  SymsAddr hi;
+  syms_uint read_offset;
+  syms_bool last_range_emitted;
+} SymsRangeIterPdb;
+
+SYMS_COMPILER_ASSERT(sizeof(SymsRangePdb) <= sizeof(SymsRangeImpl));
+SYMS_COMPILER_ASSERT(sizeof(SymsRangeIterPdb) <= sizeof(SymsRangeIterImpl));
+SYMS_COMPILER_ASSERT(sizeof(pdb_type) <= sizeof(SymsTypeImpl));
+SYMS_COMPILER_ASSERT(sizeof(pdb_mod) <= sizeof(SymsModImpl));
+SYMS_COMPILER_ASSERT(sizeof(SymsDebugInfoPdb) <= sizeof(SymsDebugInfoImpl));
+SYMS_COMPILER_ASSERT(sizeof(SymsModIterPdb) <= sizeof(SymsModIterImpl));
+SYMS_COMPILER_ASSERT(sizeof(SymsLineIterPdb) <= sizeof(SymsLineIterImpl));
+SYMS_COMPILER_ASSERT(sizeof(SymsMemberIterPdb) <= sizeof(SymsMemberIterImpl));
+SYMS_COMPILER_ASSERT(sizeof(SymsGlobalIterPdb) <= sizeof(SymsGlobalIterImpl));
+SYMS_COMPILER_ASSERT(sizeof(SymsLocalDataIterPdb) <= sizeof(SymsLocalDataIterImpl));
+SYMS_COMPILER_ASSERT(sizeof(SymsFileIterPdb) <= sizeof(SymsFileIterImpl));
+SYMS_COMPILER_ASSERT(sizeof(SymsProcIterPdb) <= sizeof(SymsProcIterImpl));
+SYMS_COMPILER_ASSERT(sizeof(SymsArgIterPdb) <= sizeof(SymsArgIterImpl));
+SYMS_COMPILER_ASSERT(sizeof(SymsLocalIterPdb) <= sizeof(SymsLocalIterImpl));
+SYMS_COMPILER_ASSERT(sizeof(SymsInlineIterPdb) <= sizeof(SymsInlineIterImpl));
+SYMS_COMPILER_ASSERT(SYMS_LOCATION_IMPLICIT_VALUE_MAX >= PDB_LOCATION_IMPLICIT_VALUE_MAX);
+SYMS_COMPILER_ASSERT(SYMS_MEMBER_ENUM_MAX >= PDB_NUMERIC_MAX);
+SYMS_COMPILER_ASSERT(SYMS_CONST_VALUE_MAX >= PDB_NUMERIC_MAX);
+SYMS_COMPILER_ASSERT(sizeof(SymsStringRefImpl) >= sizeof(pdb_pointer));
+
+typedef PDB_MEMREAD_SIG(SymsPdbMemreadSig);
+typedef PDB_REGREAD_SIG(SymsPdbRegreadSig);
+typedef PDB_REGWRITE_SIG(SymsPdbRegwriteSig);
+
+// ----------------------------------------------------------------------------
+
+SYMS_INTERNAL pdb_cv_itype
+syms_typeid_to_pdb(const SymsTypeID *type_id);
+
+SYMS_INTERNAL SymsTypeID
+syms_typeid_for_pdb(pdb_cv_itype itype);
+
+SYMS_INTERNAL syms_int
+syms_typeid_cmp_pdb(pdb_cv_itype l, pdb_cv_itype r);
+
+// ----------------------------------------------------------------------------
+
+SYMS_INTERNAL SymsStringRef
+syms_string_ref_pdb(pdb_pointer *pointer);
+
+SYMS_INTERNAL syms_bool
+syms_proc_from_pdb_proc(SymsDebugInfoPdb *debug_info, pdb_proc *proc, SymsProc *proc_out);
+
+SYMS_INTERNAL void
+syms_const_convert_from_pdb(pdb_const_value *pdb_const, SymsConst *const_out);
+
+SYMS_INTERNAL syms_bool
+syms_pdb_type_to_syms_type(pdb_context *pdb, pdb_type *type, SymsType *type_out);
+
+SYMS_INTERNAL
+PDB_REGREAD_SIG(syms_regread_pdb);
+
+#endif /* SYMS_SYMS_API_INCLUDE_H */
+
