RAD symbol library

diff --git a/Engine/Source/Runtime/SymsLib/src/regs/syms_regs.c b/Engine/Source/Runtime/SymsLib/src/regs/syms_regs.c
new file mode 100644
--- /dev/null
+++ b/Engine/Source/Runtime/SymsLib/src/regs/syms_regs.c
@@ -0,0 +1,564 @@
+// Copyright Epic Games, Inc. All Rights Reserved.
+/******************************************************************************
+ * File   : syms_regs.c                                                       *
+ * Author : Nikita Smith                                                      *
+ * Created: 2020/06/10                                                        *
+ * Purpose: Generic register interface that supports x86/x64 and ARM/ARM64    *
+ ******************************************************************************/
+
+#include "syms_regs_data.h"
+
+// NOTE(PM): We don't try to match any particular OS/CPU struct layout here, we define our own.
+// We store all the bytes individually. bytes[0] is the MSB, bytes[N-1] is the LSB. (always big endian)
+//
+
+#if 0
+SYMS_INTERNAL const SymsRegDesc *
+syms_regs_get_regdesc(SymsArch arch)
+{
+    SymsRegDesc *result = 0;
+    if (arch < SYMS_ARCH_COUNT)
+    {
+        result = g_syms_reg_desc[arch];
+    }
+    return result;
+}
+
+SYMS_INTERNAL U32
+syms_regs_get_desc_size(SymsArch arch)
+{
+    U32 result = 0;
+    if (arch < SYMS_ARCH_COUNT)
+    {
+        result = g_syms_reg_desc_size[arch];
+    }
+    return result;
+}
+#endif
+
+SYMS_API SymsRegID
+syms_regs_get_first_regid_ex(SymsArch arch)
+{
+    syms_uint result = 0;
+    if (arch < SYMS_ARCH_COUNT)
+    {
+        result = 1;
+    }
+    return result;
+}
+
+SYMS_API SymsRegID
+syms_regs_get_last_regid_ex(SymsArch arch)
+{
+    syms_uint result = 0;
+    if (arch < SYMS_ARCH_COUNT)
+    {
+        result = g_syms_reg_desc_count[arch] - 1;
+    }
+    return result;
+}
+
+SYMS_API SymsString
+syms_regs_get_name_ex(SymsArch arch, SymsRegID id)
+{
+  const SymsRegInfo *reg_table = reg_table_from_arch(arch);
+  syms_uint reg_count = g_syms_reg_desc_count[arch];
+  SymsString name = syms_string_init(0,0);
+  if (id < reg_count) {
+    name = reg_table[id].name;
+  }
+  return name;
+}
+
+SYMS_API SymsTypeKind
+syms_regs_get_type_ex(SymsArch arch, SymsRegID id)
+{
+  const SymsRegInfo *reg_table = reg_table_from_arch(arch);
+  syms_uint reg_count = g_syms_reg_desc_count[arch];
+  SymsTypeKind kind = SYMS_TYPE_NULL;
+  if (id < reg_count) {
+    kind = reg_table[id].kind;
+  }
+  return kind;
+}
+
+SYMS_API SymsRegClass
+syms_regs_get_reg_class_ex(SymsArch arch, SymsRegID id)
+{
+  const SymsRegInfo *reg_table = reg_table_from_arch(arch);
+  syms_uint reg_count = g_syms_reg_desc_count[arch];
+  SymsRegClass reg_class = SYMS_REG_CLASS_NULL;
+  if (id < reg_count) {
+    reg_class = reg_table[id].group;
+  }
+  return reg_class;
+}
+
+SYMS_API SymsRegID
+syms_regs_get_ip_ex(SymsArch arch)
+{
+  switch (arch) {
+  case SYMS_ARCH_X86: return SYMS_REG_X86_eip;
+  case SYMS_ARCH_X64: return SYMS_REG_X64_rip;
+  default:            return SYMS_REG_null;
+  }
+}
+
+SYMS_API SymsRegID
+syms_regs_get_sp_ex(SymsArch arch)
+{
+  switch (arch) {
+  case SYMS_ARCH_X86: return SYMS_REG_X86_esp;
+  case SYMS_ARCH_X64: return SYMS_REG_X64_rsp;
+  default:            return SYMS_REG_null;
+  }
+}
+
+SYMS_API SymsRegAliasInfo
+syms_regs_get_alias_info(SymsArch arch, SymsRegID id)
+{
+  const SymsRegInfo *reg_table = reg_table_from_arch(arch);
+  syms_uint reg_count = g_syms_reg_desc_count[arch];
+  SymsRegAliasInfo result;
+  if (id < reg_count) {
+    result.aliasee = reg_table[id].alias;
+    result.bit_shift = reg_table[id].bit_shift;
+    result.bit_count = reg_table[id].bit_count;
+  } else {
+    result.aliasee = 0;
+    result.bit_shift = 0;
+    result.bit_count = 0;
+  }
+  return result;
+}
+
+SYMS_INTERNAL U32
+syms_regs_get_regsdata_size(SymsArch arch)
+{
+    U32 result = 0;
+    if (arch < SYMS_ARCH_COUNT)
+    {
+        result = g_syms_reg_data_size[arch];
+    }
+    return result;
+}
+
+REG_API SymsRegs * 
+syms_regs_init(SymsArch arch, void *bf, U32 bf_size)
+{
+    syms_uint size = syms_regs_size_for_arch(arch);
+    SymsRegs *regs = 0;
+    if (size > 0 && size <= bf_size) {
+        syms_uint descs_size        = reg_table_size_from_arch(arch);
+        syms_uint data_size         = syms_regs_get_regsdata_size(arch);
+        syms_uint reg_info_count    = g_syms_reg_desc_count[arch];
+        const SymsRegInfo *reg_info = reg_table_from_arch(arch);
+        syms_uint i;
+        
+        regs = (SymsRegs *)bf;
+        regs->endiness = 0x1FFF;
+        regs->arch_le = SYMS_LE16((u16)arch);
+        
+        regs->version_le    = SYMS_LE16((U16)1);
+        regs->firstreg_le   = SYMS_LE32((U32)syms_regs_get_first_regid_ex(arch));
+        regs->lastreg_le    = SYMS_LE32((U32)syms_regs_get_last_regid_ex(arch));
+        regs->size_le       = SYMS_LE32((U32)size);
+
+        for (i = 0; i < reg_info_count; ++i) {
+          const SymsRegInfo *r = reg_info + i;
+          SymsRegDesc *r_out   = regs->descs + i;
+
+          SYMS_ASSERT(r->name.len <= sizeof(regs->descs[i].name));
+          syms_memcpy(regs->descs[i].name, r->name.data, r->name.len);
+          r_out->regclass     = SYMS_LE16((u16)r->group);
+          r_out->varOffset_le = SYMS_LE32((u16)r->offset);
+          r_out->bitpos_le    = SYMS_LE16((u16)r->bit_shift);
+          r_out->bitcount_le  = SYMS_LE16((u16)r->bit_count);
+        }
+
+        syms_memset(((U8*)&regs->descs[0]) + descs_size, 0, data_size);
+    }
+    
+    return regs;
+}
+
+REG_API U32
+syms_regs_size_for_arch(SymsArch arch)
+{
+    U32 result = 0;
+    if (arch < SYMS_ARCH_COUNT) {
+        result = sizeof(SymsRegs) + syms_regs_get_regsdata_size(arch) + reg_table_size_from_arch(arch);
+    }
+    return result;
+}
+
+REG_API SymsRegs *
+syms_regs_clone(const SymsRegs *regs, void *buffer, U32 buffer_max)
+{
+    U32 size = syms_regs_get_size(regs);
+    SymsRegs *result = 0;
+    if (size <= buffer_max) {
+        result = (SymsRegs *)syms_memcpy(buffer, regs, size);
+    }
+    return result;
+}
+
+REG_API SymsArch 
+syms_regs_get_arch(const SymsRegs *regs)
+{
+    SymsArch result = SYMS_ARCH_COUNT;
+    if (regs) {
+        result = (SymsArch)SYMS_LE16(regs->arch_le);
+    }
+    return result;
+}
+
+REG_API const SymsRegDesc *
+syms_regs_get_regdesc(const SymsRegs *regs, SymsRegID id)
+{
+    const SymsRegDesc *result = 0;
+    SymsRegID first_reg = syms_regs_get_first_regid(regs);
+    SymsRegID last_reg = syms_regs_get_last_regid(regs);
+    if (first_reg <= id && id <= last_reg) {
+        result = regs->descs + ((U16)id - (U16)first_reg);
+    }
+    return result;
+}
+
+REG_API U32 
+syms_regs_get_size(const SymsRegs *regs)
+{
+    return SYMS_LE32(regs->size_le);
+}
+
+REG_API U32 
+syms_regs_get_version(const SymsRegs *regs)
+{
+    return SYMS_LE32(regs->version_le);
+}
+
+REG_API SymsRegID 
+syms_regs_get_first_regid(const SymsRegs *regs)
+{
+    return (SymsRegID)SYMS_LE32(regs->firstreg_le);
+}
+
+REG_API SymsRegID 
+syms_regs_get_last_regid(const SymsRegs *regs)
+{
+    return (SymsRegID)SYMS_LE32(regs->lastreg_le);
+}
+
+REG_API SymsRegID 
+syms_regs_get_ip_regid(const SymsRegs *regs)
+{
+    switch (syms_regs_get_arch(regs)) {
+        case SYMS_ARCH_X86: return SYMS_REG_X86_eip;
+        case SYMS_ARCH_X64: return SYMS_REG_X64_rip;
+        default:            return 0;
+    }
+}
+
+REG_API SymsRegID 
+syms_regs_get_sp_regid(const SymsRegs *regs)
+{
+    switch (syms_regs_get_arch(regs)) {
+        case SYMS_ARCH_X86: return SYMS_REG_X86_esp;
+        case SYMS_ARCH_X64: return SYMS_REG_X64_rsp;
+        default:            return 0;
+    }
+}
+
+SYMS_INTERNAL SymsRegsPtr 
+syms_regs_query(const SymsRegs *regs, SymsRegID id)
+{
+    SymsRegsPtr result;
+    
+    // HACK: The X87 FPU handles ST/MMX registers in a really annoying way.
+    // They're overlaid, but it's not a static mapping; it changes
+    // based on the stack TOP pointer. Note that lldb gets this wrong, while
+    // gdb gets it right. We hack it by never using the st0 data directly,
+    // we redirect it into the appropriate fpr register.
+    {
+        SymsArch arch = syms_regs_get_arch(regs);
+        
+        SymsRegID reg_st0 = 0;
+        SymsRegID reg_fsw = 0;
+        SymsRegID reg_fpr0 = 0;
+        switch (arch) {
+            default: break;
+            
+            case SYMS_ARCH_X64: {
+                reg_st0 = SYMS_REG_X64_st0;
+                if (id >= reg_st0 && id <= SYMS_REG_X64_st7) {
+                    reg_fsw = SYMS_REG_X64_fsw;
+                    reg_fpr0 = SYMS_REG_X64_fpr0;
+                    goto X86_64;
+                }
+            } break;
+            
+            case SYMS_ARCH_X86: {
+                reg_st0 = SYMS_REG_X86_st0;
+                if (id >= reg_st0 && id <= SYMS_REG_X86_st7) {
+                    reg_fsw = SYMS_REG_X86_fsw;
+                    reg_fpr0 = SYMS_REG_X86_fpr0;
+                    goto X86_64;
+                }
+            } break;
+            
+            X86_64:;
+            {
+                U16 fsw = syms_regs_get16(regs, reg_fsw);
+                U16 st_index = (U16)id - (U16)reg_st0;
+                U32 fpr_index = (st_index + (fsw >> 11)) & 7;
+                id = (SymsRegID)(reg_fpr0 + fpr_index);
+            }break;
+        }
+    }
+    
+    {
+        const SymsRegDesc *regdesc = syms_regs_get_regdesc(regs, id);
+        if (regdesc != NULL) {
+            SymsRegID first_reg = syms_regs_get_first_regid(regs);
+            SymsRegID last_reg = syms_regs_get_last_regid(regs);
+            U32 regdesc_count = (last_reg - first_reg) + 1;
+            U32 regdesc_size = sizeof(SymsRegDesc) * regdesc_count;
+            U8 *ptr;
+            
+            SYMS_ASSERT((regdesc->bitpos_le & 7) == 0); // we don't support non-8bit fields
+            SYMS_ASSERT((regdesc->bitcount_le & 7) == 0);
+            
+            ptr = (U8 *)regs->descs + regdesc_size;
+            result.data = ptr + regdesc->varOffset_le + regdesc->bitpos_le / 8;
+            result.bytes = regdesc->bitcount_le / 8;
+        } else {
+            result.data = 0;
+            result.bytes = 0;
+        }
+    }
+    
+    return result;
+}
+
+REG_API SymsRegID
+syms_regid_from_name(const SymsRegs *regs, char *name)
+{
+    SymsArch arch = syms_regs_get_arch(regs);
+    if (arch < SYMS_ARCH_COUNT) {
+        // TODO(nick): A hash table to speed this up?
+        const SymsRegInfo *info_table = reg_table_from_arch(arch);
+        SymsRegID regid_first = syms_regs_get_first_regid(regs);
+        SymsRegID regid_last = syms_regs_get_last_regid(regs);
+        U32 count = (U32)regid_last - (U32)regid_first + 1;
+        U32 i;
+        for (i = 0; i < count; ++i) {
+          if (syms_strcmp(info_table[i].name.data, name) == 0) {
+            return (SymsRegID)((U32)regid_first + i);
+          }
+        }
+    }
+    return 0;
+}
+
+SYMS_INTERNAL void 
+syms_regs_memcpy_swap(U8 *dst, const U8 *src, U32 n)
+{
+    while (n--) {
+        *dst++ = src[n];
+    }
+}
+
+REG_API SymsEndian
+syms_regs_get_endianess(const SymsRegs *regs)
+{
+    if (((U8 *)&regs->endiness)[0] == 0xFF) {
+        return SYMS_ENDIAN_LITTLE;
+    } else {
+        return SYMS_ENDIAN_BIG;
+    }
+}
+
+REG_API void 
+syms_regs_clear_reg(SymsRegs *regs, SymsRegID id)
+{
+    SymsRegsPtr dst = syms_regs_query(regs, id);
+    syms_memset((void *)dst.data, 0, dst.bytes);   // Zero the register out.
+}
+
+REG_API U32 
+syms_regs_get_value(const SymsRegs *regs, SymsRegID id, void *dst, U32 dst_size)
+{
+    U32 written = 0;
+    SymsRegsPtr src = syms_regs_query(regs, id);
+    if (src.bytes <= dst_size) {
+        if (syms_regs_get_endianess(regs) == SYMS_ENDIAN_HOST) {
+            syms_memcpy(dst, src.data, src.bytes);
+        } else {
+            syms_regs_memcpy_swap((U8 *)dst, (const U8 *)src.data, src.bytes);
+        }
+        written = src.bytes;
+    }
+    return written;
+}
+
+REG_API U8 
+syms_regs_get8(const SymsRegs *regs, SymsRegID id)
+{
+    U8 n = 0;
+    SYMS_ASSERT(syms_regs_query(regs, id).bytes == 1);
+    syms_regs_get_value(regs, id, &n, sizeof(n));
+    return n;
+}
+
+REG_API U16 
+syms_regs_get16(const SymsRegs *regs, SymsRegID id)
+{
+    U16 n = 0;
+    SYMS_ASSERT(syms_regs_query(regs, id).bytes == 2);
+    syms_regs_get_value(regs, id, &n, sizeof(n));
+    return n;
+}
+
+REG_API U32 
+syms_regs_get32(const SymsRegs *regs, SymsRegID id)
+{
+    U32 n = 0;
+    SYMS_ASSERT(syms_regs_query(regs, id).bytes == 4);
+    syms_regs_get_value(regs, id, &n, sizeof(n));
+    return n;
+}
+
+REG_API U64 
+syms_regs_get64(const SymsRegs *regs, SymsRegID id)
+{
+    U64 n = 0;
+    SYMS_ASSERT(syms_regs_query(regs, id).bytes == 8);
+    syms_regs_get_value(regs, id, &n, sizeof(n));
+    return n;
+}
+
+REG_API syms_bool 
+syms_regs_get_addr(const SymsRegs *regs, SymsRegID id, SymsAddr *addr_out)
+{
+    syms_bool is_result_valid = syms_false;
+    U32 bytes = syms_regs_query(regs, id).bytes;
+    is_result_valid = syms_true;
+    switch (bytes) {
+        case 1: *addr_out = syms_regs_get8(regs, id);  break;
+        case 2: *addr_out = syms_regs_get16(regs, id); break;
+        case 4: *addr_out = syms_regs_get32(regs, id); break;
+        case 8: *addr_out = syms_regs_get64(regs, id); break;
+        default: is_result_valid = syms_false;              break;
+    }
+    return is_result_valid;
+}
+
+REG_API SymsAddr 
+syms_regs_get_ip(const SymsRegs *regs)
+{
+    switch (syms_regs_get_arch(regs)) {
+        // TODO(allen): This doesn't make sense. Some OSes can't deal with changes
+        // to the CS register, so we're going to be broken for all OSes instead of
+        // just the ones that have to be??? We're changing this.
+        
+        // TODO: EIP/RIP isn't actually the instruction pointer. (really!)
+        //       It's an *offset* into the code segment.
+        //       The actual address of the running code would be CS:EIP.
+        //       (i.e. csbase+eip)
+        //       We don't handle this though, as it's *incredibly* rare (and
+        //       not even possible on some OSes) to change the CS segment base
+        //       to something non-zero. Many OSes can't even return the csbase
+        //       via their debugger API.
+        case SYMS_ARCH_X86: return syms_regs_get32(regs, SYMS_REG_X86_eip);
+        case SYMS_ARCH_X64: return syms_regs_get64(regs, SYMS_REG_X64_rip);
+        default:            return 0;
+    }
+}
+
+REG_API SymsAddr 
+syms_regs_get_sp(const SymsRegs *regs)
+{
+    switch (syms_regs_get_arch(regs)) {
+        // TODO: should be SS:EIP (ssbase+eip), same as above
+        case SYMS_ARCH_X86: return syms_regs_get32(regs, SYMS_REG_X86_esp);
+        case SYMS_ARCH_X64: return syms_regs_get64(regs, SYMS_REG_X64_rsp);
+        default:        return 0;
+    }
+}
+
+REG_API U32 
+syms_regs_set_value(SymsRegs *regs, SymsRegID id, const void *src, U32 src_size)
+{
+    U32 written = 0;
+    SymsRegsPtr dst = syms_regs_query(regs, id);
+    if (src_size > 0 && src_size <= dst.bytes) {
+        if (syms_regs_get_endianess(regs) == SYMS_ENDIAN_HOST) {
+            syms_memcpy(dst.data, src, src_size);
+        } else {
+            syms_regs_memcpy_swap((U8 *)dst.data, (const U8 *)src, src_size);
+        }
+        written = dst.bytes;
+    }
+    return written;
+}
+
+REG_API syms_bool
+syms_regs_set8(SymsRegs *regs, SymsRegID id, U8 value)
+{
+    SYMS_ASSERT(syms_regs_query(regs, id).bytes == 1);
+    return syms_regs_set_value(regs, id, &value, sizeof(value));
+}
+
+REG_API syms_bool 
+syms_regs_set16(SymsRegs *regs, SymsRegID id, U16 value)
+{
+    SYMS_ASSERT(syms_regs_query(regs, id).bytes == 2);
+    return syms_regs_set_value(regs, id, &value, sizeof(value));
+}
+
+REG_API syms_bool 
+syms_regs_set32(SymsRegs *regs, SymsRegID id, U32 value)
+{
+    SYMS_ASSERT(syms_regs_query(regs, id).bytes == 4);
+    return syms_regs_set_value(regs, id, &value, sizeof(value));
+}
+
+REG_API syms_bool 
+syms_regs_set64(SymsRegs *regs, SymsRegID id, U64 value)
+{
+    SYMS_ASSERT(syms_regs_query(regs, id).bytes == 8);
+    return syms_regs_set_value(regs, id, &value, sizeof(value));
+}
+
+REG_API syms_bool 
+syms_regs_set_addr(SymsRegs *regs, SymsRegID id, SymsAddr value)
+{
+    U32 bytes = syms_regs_query(regs, id).bytes;
+    switch (bytes) {
+        case 1:     return syms_regs_set8(regs, id, (U8)value);
+        case 2:     return syms_regs_set16(regs, id, (U16)value);
+        case 4:     return syms_regs_set32(regs, id, (U32)value);
+        case 8:     return syms_regs_set64(regs, id, (U64)value);
+        default:    return syms_false;
+    }
+}
+
+REG_API syms_bool
+syms_regs_set_ip(SymsRegs *regs, SymsAddr value)
+{
+    switch (syms_regs_get_arch(regs)) {
+        case SYMS_ARCH_X86: return syms_regs_set32(regs, SYMS_REG_X86_eip, (U32)value);
+        case SYMS_ARCH_X64: return syms_regs_set64(regs, SYMS_REG_X64_rip, (U64)value);
+        default:            return syms_false;
+    }
+}
+
+REG_API syms_bool
+syms_regs_set_sp(SymsRegs *regs, SymsAddr value)
+{
+    switch (syms_regs_get_arch(regs)) {
+        case SYMS_ARCH_X86: return syms_regs_set32(regs, SYMS_REG_X86_esp, (U32)value);
+        case SYMS_ARCH_X64: return syms_regs_set64(regs, SYMS_REG_X64_rsp, (U64)value);
+        default:        return syms_false;
+    }
+}
diff --git a/Engine/Source/Runtime/SymsLib/src/regs/syms_regs.h b/Engine/Source/Runtime/SymsLib/src/regs/syms_regs.h
new file mode 100644
--- /dev/null
+++ b/Engine/Source/Runtime/SymsLib/src/regs/syms_regs.h
@@ -0,0 +1,179 @@
+// Copyright Epic Games, Inc. All Rights Reserved.
+#ifndef SYMS_REGS_INCLUDE_H
+#define SYMS_REGS_INCLUDE_H
+
+/******************************************************************************
+ * File   : syms_regs.h                                                       *
+ * Author : Nikita Smith                                                      *
+ * Created: 2020/06/10                                                        *
+ * Purpose: Generic register interface that supports x86/x64                  *
+ ******************************************************************************/
+
+#ifndef REG_API
+#define REG_API SYMS_API
+#endif
+
+typedef U32 SymsRegID;
+
+typedef struct SymsRegDesc 
+{
+    // register name/alias
+    char name[8];
+    
+    // register class
+    U16 regclass;
+    
+    // register contained within a variable at this byte offset into the register file 
+    U16 varOffset_le;
+    
+    // bit offset within the containing variable 
+    U16 bitpos_le;
+    
+    // bit size within the containing variable 
+    U16 bitcount_le;
+} SymsRegDesc;
+SYMS_COMPILER_ASSERT(sizeof(SymsRegDesc) == 16);
+
+typedef struct SymsRegsPtr 
+{
+    U8 *data;
+    U16 bytes;
+} SymsRegsPtr;
+
+typedef struct SymsRegs 
+{
+    U16 endiness;
+    
+    // SymsArch type
+    U16 arch_le;
+    
+    // Regs version
+    U16 version_le;
+    
+    // Lowest register
+    U32 firstreg_le;
+    
+    // Highest register
+    U32 lastreg_le;
+    
+    /* Total register file size */
+    U32 size_le;
+    
+    SymsRegDesc descs[0];
+} SymsRegs;
+
+typedef struct
+{
+    SymsRegs *regs;
+    
+    U32 desc_count;
+    SymsRegDesc *desc;
+    
+    U32 values_size;
+    U8 *values;
+} SymsRegsInfo;
+
+#define SYMS_REG_ID_ARCH_MASK(x) ((U32)(x) >> 12)
+
+syms_bool
+syms_get_table_info(SymsArch arch, SymsRegDesc **desc_out, U32 *desc_size_out, U32 *regs_size_out, U32 *ver_out, U32 *regid_min_out, U32 *regid_max_out);
+
+REG_API U32
+syms_regs_size_for_arch(SymsArch arch);
+
+REG_API SymsRegs *
+syms_regs_init(SymsArch arch, void *bf, U32 bf_max);
+
+REG_API SymsRegs *
+syms_regs_clone(const SymsRegs *regs, void *buffer, U32 buffer_max);
+
+// Returns the CPU architecture for a register file.
+REG_API SymsArch 
+syms_regs_get_arch(const SymsRegs *regs);
+
+// Returns the size of a register file, in bytes.
+REG_API U32 
+syms_regs_get_size(const SymsRegs *regs);
+
+// Returns the version number of a register file.
+REG_API U32 
+syms_regs_get_version(const SymsRegs *regs);
+
+// Returns the lowest register ID for a file. 
+REG_API SymsRegID 
+syms_regs_get_first_regid(const SymsRegs *regs);
+
+// Returns the highest register ID for a file.
+REG_API SymsRegID 
+syms_regs_get_last_regid(const SymsRegs *regs);
+
+// Returns the register ID for the instruction pointer.
+REG_API SymsRegID 
+syms_regs_get_ip_regid(const SymsRegs *regs);
+
+// Returns the register ID for the stack pointer.
+REG_API SymsRegID
+syms_regs_get_sp_regid(const SymsRegs *regs);
+
+// Returns the description for a register file & index.
+REG_API const SymsRegDesc *
+syms_regs_get_regdesc(const SymsRegs *regs, SymsRegID id);
+
+REG_API SymsRegID
+syms_regid_from_name(const SymsRegs *regs, char *name);
+
+REG_API U32 
+syms_regs_get_value(const SymsRegs *regs, SymsRegID id, void *dst, U32 dst_size);
+
+REG_API U32 
+syms_regs_set_value(SymsRegs *regs, SymsRegID id, const void *src, U32 src_size);
+
+REG_API void 
+syms_regs_clear_reg(SymsRegs *regs, SymsRegID id);
+
+REG_API U8  
+syms_regs_get8(const SymsRegs *regs, SymsRegID id);
+
+REG_API U16 
+syms_regs_get16(const SymsRegs *regs, SymsRegID id);
+
+REG_API U32 
+syms_regs_get32(const SymsRegs *regs, SymsRegID id);
+
+REG_API U64 
+syms_regs_get64(const SymsRegs *regs, SymsRegID id);
+
+REG_API syms_bool 
+syms_regs_get_addr(const SymsRegs *regs, SymsRegID id, SymsAddr *addr_out);
+
+REG_API syms_bool 
+syms_regs_set8(SymsRegs *regs, SymsRegID id, U8 value);
+
+REG_API syms_bool 
+syms_regs_set16(SymsRegs *regs, SymsRegID id, U16 value);
+
+REG_API syms_bool 
+syms_regs_set32(SymsRegs *regs, SymsRegID id, U32 value);
+
+REG_API syms_bool 
+syms_regs_set64(SymsRegs *regs, SymsRegID id, U64 value);
+
+REG_API syms_bool 
+syms_regs_set_addr(SymsRegs *regs, SymsRegID id, SymsAddr value);
+
+REG_API SymsAddr 
+syms_regs_get_ip(const SymsRegs *regs);
+
+REG_API SymsAddr 
+syms_regs_get_sp(const SymsRegs *regs);
+
+REG_API syms_bool 
+syms_regs_set_ip(SymsRegs *regs, SymsAddr value);
+
+REG_API syms_bool 
+syms_regs_set_sp(SymsRegs *regs, SymsAddr value);
+
+REG_API SymsEndian
+syms_regs_get_endianess(const SymsRegs *regs);
+
+#endif /* SYMS_REGS_INCLUDE_H */
diff --git a/Engine/Source/Runtime/SymsLib/src/regs/syms_regs_intel.c b/Engine/Source/Runtime/SymsLib/src/regs/syms_regs_intel.c
new file mode 100644
--- /dev/null
+++ b/Engine/Source/Runtime/SymsLib/src/regs/syms_regs_intel.c
@@ -0,0 +1,225 @@
+// Copyright Epic Games, Inc. All Rights Reserved.
+// Constructs an x87 tag value given a 10-byte float.
+SYMS_INTERNAL U32 
+syms_regs_make_tag(const void *fp10)
+{
+  U16 exponent = ((const U16 *)fp10)[4] & 0x7fff;
+  U64 bits = ((const U64 *)fp10)[0];
+
+  U64 j = bits >> 63;
+  U64 frac = bits & ~((U64)1 << 63);
+
+  if (exponent == 0x7fff) {
+    return 2; // special
+  } else if (exponent == 0) {
+    return (frac == 0 && !j) ? /* zero */ 1 : /* special */ 2;
+  } else {
+    return j ? /* valid */ 0 : /*special */ 2;
+  }
+}
+
+SYMS_INTERNAL U16 
+syms_regs_fix_tag_word(SymsIntelFXSave *fpu)
+{
+  // The FXSAVE/XSAVE instructions do not save out the correct FPU tag word.
+  // But, we can recreate the correct value by looking at the actual FP data.
+  // (see FXSAVE instruction in Intel IA32 Software Developers Manual)
+
+  U16 realFtw = 0;
+  U32 top = (fpu->fsw >> 11) & 7;
+  u32 fpr;
+  for (fpr = 0; fpr < 8; fpr++) {
+    U32 tag;
+
+    if (fpu->ftw & (1 << fpr)) {
+      U32 st = (fpr - top) & 7;
+      tag = syms_regs_make_tag((const void *)&fpu->st[st*16]);
+    } else {
+      tag = 3; // empty
+    }
+
+    realFtw |= tag << (2 * fpr);
+  }
+
+  return realFtw;
+}
+
+SYMS_INTERNAL U16 
+syms_regs_break_tag_word(U16 ftw)
+{
+  // Generates the compact tag word used by FXSAVE/XSAVE.
+  U16 compact = 0;
+  U32 fpr;
+  for (fpr = 0; fpr < 8; fpr++) {
+    U32 tag = (ftw >> (fpr * 2)) & 3;
+    if (tag != 3) {
+      compact |= (1 << fpr);
+    }
+  }
+  return compact;
+}
+
+// X64 --------------------------------------------------------------------
+
+SYMS_COMPILER_ASSERT(SYMS_REG_X64_st7 == SYMS_REG_X64_st0 + 7);
+SYMS_COMPILER_ASSERT(SYMS_REG_X86_xmm7 == SYMS_REG_X86_xmm0 + 7);
+SYMS_COMPILER_ASSERT(SYMS_REG_X64_xmm15 == SYMS_REG_X64_xmm0 + 15);
+SYMS_COMPILER_ASSERT(SYMS_REG_X64_ymm15 == SYMS_REG_X64_ymm0 + 15);
+SYMS_COMPILER_ASSERT(SYMS_REG_X64_ymm15 == SYMS_REG_X64_ymm0 + 15);
+
+SYMS_INTERNAL void 
+syms_regs_x64_fxsave_get_regs(SymsRegs *dest, SymsIntelFXSave *src)
+{
+  u32 n;
+  U16 ftw = syms_regs_fix_tag_word(src);
+  syms_regs_set_value(dest, SYMS_REG_X64_fcw, &src->fcw, sizeof(src->fcw));
+  syms_regs_set_value(dest, SYMS_REG_X64_fsw, &src->fsw, sizeof(src->fsw));
+  syms_regs_set_value(dest, SYMS_REG_X64_ftw, &ftw, sizeof(src->ftw));
+  syms_regs_set_value(dest, SYMS_REG_X64_fop, &src->fop, sizeof(src->fop));
+  syms_regs_set_value(dest, SYMS_REG_X64_fip, &src->u.format64.fip, sizeof(src->u.format64.fip));
+  syms_regs_set_value(dest, SYMS_REG_X64_fdp, &src->u.format64.fdp, sizeof(src->u.format64.fdp));
+  syms_regs_set_value(dest, SYMS_REG_X64_mxcsr, &src->mxcsr, sizeof(src->mxcsr));
+  for (n = 0; n < 8; n++) {
+    syms_regs_set_value(dest, (SymsRegID)(SYMS_REG_X64_st0 + n), &src->st[n*16], sizeof(src->st[n*16]));
+  }
+  for (n = 0; n < 16; n++) {
+    syms_regs_set_value(dest, (SymsRegID)(SYMS_REG_X64_xmm0 + n), &src->xmm[n*16], sizeof(src->xmm[n*16]));
+  }
+}
+
+SYMS_INTERNAL void 
+syms_regs_x64_fxsave_put_regs(SymsIntelFXSave *dest, struct SymsRegs *src)
+{
+  U32 n;
+  syms_regs_get_value(src, SYMS_REG_X64_fcw, &dest->fcw, sizeof(dest->fcw));
+  syms_regs_get_value(src, SYMS_REG_X64_fsw, &dest->fsw, sizeof(dest->fsw));
+  syms_regs_get_value(src, SYMS_REG_X64_ftw, &dest->ftw, sizeof(dest->ftw));
+  syms_regs_get_value(src, SYMS_REG_X64_fop, &dest->fop, sizeof(dest->fop));
+  syms_regs_get_value(src, SYMS_REG_X64_fip, &dest->u.format64.fip, sizeof(dest->u.format64.fip));
+  syms_regs_get_value(src, SYMS_REG_X64_fdp, &dest->u.format64.fdp, sizeof(dest->u.format64.fdp));
+  syms_regs_get_value(src, SYMS_REG_X64_mxcsr, &dest->mxcsr, sizeof(dest->mxcsr));
+  for (n = 0; n < 8; n++) {
+    syms_regs_get_value(src, (SymsRegID)(SYMS_REG_X64_st0+n), &dest->st[n*16], sizeof(dest->st[0]));
+  }
+  for (n = 0; n < 16; n++) {
+    syms_regs_get_value(src, (SymsRegID)(SYMS_REG_X64_xmm0+n), &dest->xmm[n*16], sizeof(dest->xmm[0]));
+  }
+  dest->ftw = syms_regs_break_tag_word(dest->ftw);
+}
+
+SYMS_INTERNAL void 
+syms_regs_x64_xsave_get_regs(SymsRegs *dest, SymsIntelXSave *src)
+{
+  syms_regs_x64_fxsave_get_regs(dest, &src->fxsave);
+
+  // Check if YMM regs exist
+  if (src->xsave_hdr.xstate_bv & 4) {
+    U8 reg[32];
+    U32 n;
+    for (n = 0; n < 16; n++) {
+      syms_memcpy(&reg[0], &src->fxsave.xmm[n*16], 16); // construct the two halves
+      syms_memcpy(&reg[16], &src->ymmh[n*16], 16);
+      syms_regs_set_value(dest, (SymsRegID)(SYMS_REG_X64_ymm0 + n), reg, sizeof(reg));
+    }
+  }
+}
+
+SYMS_INTERNAL void 
+syms_regs_x64_xsave_put_regs(SymsIntelXSave *dest, struct SymsRegs *src)
+{
+  U8 reg[32];
+  U32 n;
+
+  syms_regs_x64_fxsave_put_regs(&dest->fxsave, src);
+
+  dest->xsave_hdr.xstate_bv |= 7; // Indicate we changed FPU/SSE/AVX
+
+  // Write YMM upper halves
+  for (n = 0; n < 16; n++) {
+    syms_regs_get_value(src, (SymsRegID)(SYMS_REG_X64_ymm0 + n), reg, sizeof(reg));
+    syms_memcpy(&dest->ymmh[n*16], &reg[16], 16);
+  }
+}
+
+// X86 --------------------------------------------------------------------
+
+SYMS_COMPILER_ASSERT(SYMS_REG_X86_ymm7 == SYMS_REG_X86_ymm0 + 7);
+
+SYMS_INTERNAL void 
+syms_regs_x86_fxsave_get_regs(SymsRegs *dest, SymsIntelFXSave *src) 
+{
+  U32 n;
+  U16 ftw = syms_regs_fix_tag_word(src);
+  syms_regs_set_value(dest, SYMS_REG_X86_fcw, &src->fcw, sizeof(src->fcw));
+  syms_regs_set_value(dest, SYMS_REG_X86_fsw, &src->fsw, sizeof(src->fsw));
+  syms_regs_set_value(dest, SYMS_REG_X86_ftw, &ftw, sizeof(ftw));
+  syms_regs_set_value(dest, SYMS_REG_X86_fop, &src->fop, sizeof(src->fop));
+  syms_regs_set_value(dest, SYMS_REG_X86_fip, &src->u.format32.fip, sizeof(src->u.format32.fip));
+  syms_regs_set_value(dest, SYMS_REG_X86_fcs, &src->u.format32.fcs, sizeof(src->u.format32.fcs));
+  syms_regs_set_value(dest, SYMS_REG_X86_fdp, &src->u.format32.fdp, sizeof(src->u.format32.fdp));
+  syms_regs_set_value(dest, SYMS_REG_X86_fds, &src->u.format32.fds, sizeof(src->u.format32.fds));
+  syms_regs_set_value(dest, SYMS_REG_X86_mxcsr, &src->mxcsr, sizeof(src->mxcsr));
+  for (n = 0; n < 8; n++) {
+    syms_regs_set_value(dest, (SymsRegID)(SYMS_REG_X86_st0+n), &src->st[n*16], sizeof(src->st[0]));
+  }
+  for (n = 0; n < 8; n++) {
+    syms_regs_set_value(dest, (SymsRegID)(SYMS_REG_X86_xmm0+n), &src->xmm[n*16], sizeof(src->xmm[0]));
+  }
+}
+
+SYMS_INTERNAL void 
+syms_regs_x86_fxsave_put_regs(SymsIntelFXSave *dest, struct SymsRegs *src)
+{
+  U32 n;
+  syms_regs_get_value(src, SYMS_REG_X86_fcw, &dest->fcw, sizeof(dest->fcw));
+  syms_regs_get_value(src, SYMS_REG_X86_fsw, &dest->fsw, sizeof(dest->fsw));
+  syms_regs_get_value(src, SYMS_REG_X86_ftw, &dest->ftw, sizeof(dest->ftw));
+  syms_regs_get_value(src, SYMS_REG_X86_fop, &dest->fop, sizeof(dest->fop));
+  syms_regs_get_value(src, SYMS_REG_X86_fip, &dest->u.format32.fip, sizeof(dest->u.format32.fip));
+  syms_regs_get_value(src, SYMS_REG_X86_fcs, &dest->u.format32.fcs, sizeof(dest->u.format32.fcs));
+  syms_regs_get_value(src, SYMS_REG_X86_fdp, &dest->u.format32.fdp, sizeof(dest->u.format32.fdp));
+  syms_regs_get_value(src, SYMS_REG_X86_fds, &dest->u.format32.fds, sizeof(dest->u.format32.fds));
+  syms_regs_get_value(src, SYMS_REG_X86_mxcsr, &dest->mxcsr, sizeof(dest->mxcsr));
+  for (n = 0; n < 8; n++) {
+    syms_regs_get_value(src, (SymsRegID)(SYMS_REG_X86_st0+n), &dest->st[n*16], sizeof(dest->st[0]));
+  }
+  for (n = 0; n < 8; n++) {
+    syms_regs_get_value(src, (SymsRegID)(SYMS_REG_X86_xmm0+n), &dest->xmm[n*16], sizeof(dest->xmm[0]));
+  }
+
+  dest->ftw = syms_regs_break_tag_word(dest->ftw);
+}
+
+SYMS_INTERNAL void 
+syms_regs_x86_xsave_get_regs(SymsRegs *dest, SymsIntelXSave *src)
+{
+  syms_regs_x86_fxsave_get_regs(dest, &src->fxsave);
+
+  // Check if YMM regs exist
+  if (src->xsave_hdr.xstate_bv & 4) {
+    U8 reg[32];
+    U32 n;
+    for (n = 0; n < 8; n++) {
+      syms_memcpy(&reg[0], &src->fxsave.xmm[n*16], 16); // construct the two halves
+      syms_memcpy(&reg[16], &src->ymmh[n*16], 16);
+      syms_regs_set_value(dest, (SymsRegID)(SYMS_REG_X86_ymm0+n), reg, sizeof(reg));
+    }
+  }
+}
+
+SYMS_INTERNAL void 
+syms_regs_x86_xsave_put_regs(SymsIntelXSave *dest, struct SymsRegs *src)
+{
+  U8 reg[32];
+  U32 n;
+
+  syms_regs_x86_fxsave_put_regs(&dest->fxsave, src);
+
+  dest->xsave_hdr.xstate_bv |= 7; // Indicate we changed FPU/SSE/AVX
+
+  // Write YMM upper halves
+  for (n = 0; n < 8; n++) {
+    syms_regs_get_value(src, (SymsRegID)(SYMS_REG_X86_ymm0+n), reg, sizeof(reg));
+    syms_memcpy(&dest->ymmh[n*16], &reg[16], 16);
+  }
+}
diff --git a/Engine/Source/Runtime/SymsLib/src/regs/syms_regs_intel.h b/Engine/Source/Runtime/SymsLib/src/regs/syms_regs_intel.h
new file mode 100644
--- /dev/null
+++ b/Engine/Source/Runtime/SymsLib/src/regs/syms_regs_intel.h
@@ -0,0 +1,94 @@
+// Copyright Epic Games, Inc. All Rights Reserved.
+#ifndef SYMS_REGS_INTEL_H
+#define SYMS_REGS_INTEL_H
+
+// These FPU structs match the CPU's store layouts exactly.
+
+typedef struct SymsIntelFSave {
+  U16 fcw, _pad0;            // control word
+  U16 fsw, _pad1;            // status word
+  U16 ftw, _pad2;            // tag word
+  U32 fip;                   // EIP
+  U16 fcs;                   // CS
+  U16 fop;                   // last opcode
+  U32 fdp;                   // data pointer
+  U16 fds, _pad3;            // data selector
+  U8 st[80];                 // 8 * fpreg
+} SymsIntelFSave;
+SYMS_COMPILER_ASSERT(sizeof(SymsIntelFSave) == 108);
+
+typedef struct SymsIntelFXSave {
+  U16 fcw;                   // control word
+  U16 fsw;                   // status word
+  U16 ftw;                   // tag word
+  U16 fop;                   // last opcode
+  union {
+    struct {
+      U64 fip;           // RIP
+      U64 fdp;           // data pointer
+    } format64;
+    struct {
+      U32 fip;           // EIP
+      U16 fcs, _pad0;    // CS
+      U32 fdp;           // data pointer
+      U16 fds, _pad1;    // data selector
+    } format32;
+  } u;
+  U32 mxcsr;                 // MXCSR Register State
+  U32 mxcsr_mask;            // MXCSR Mask
+  U8 st[128];                // 8 * (fpreg + padding)
+  U8 xmm[256];               // 16 * xmmreg
+  U8 reserved[96];
+} SymsIntelFXSave;
+SYMS_COMPILER_ASSERT(sizeof(SymsIntelFXSave) == 512);
+
+typedef struct SymsIntelXSaveHdr {
+  U64 xstate_bv;
+  U8 reserved[56];
+} SymsIntelXSaveHdr;
+SYMS_COMPILER_ASSERT(sizeof(SymsIntelXSaveHdr) == 64);
+
+typedef struct SymsIntelXSave {
+  SymsIntelFXSave fxsave;
+  SymsIntelXSaveHdr xsave_hdr;
+
+  // NOTE:
+  // Technically the layout and size of the XSAVE struct is not defined
+  // beyond the header.
+  // The idea is that you should query CPUID to find out what features
+  // the CPU has, and the size+offset of each one.
+  // GDB/LLDB, and even the Linux kernel, assume that the YMMH registers
+  // are always at a fixed location. The Intel manuals suggest this to be
+  // the case too, but the AMD manuals don't.
+  // Anyway. This will do for now, in future we may need to query CPUID
+  // ourselves to know the layout here.
+  U8 ymmh[256];              // 16 * 16 bytes for each YMMH-reg
+} SymsIntelXSave;
+
+struct SymsRegs;
+
+SYMS_INTERNAL void 
+syms_regs_x64_fxsave_get_regs(struct SymsRegs *dest, SymsIntelFXSave *src);
+
+SYMS_INTERNAL void 
+syms_regs_x64_xsave_get_regs(struct SymsRegs *dest, SymsIntelXSave *src);
+
+SYMS_INTERNAL void 
+syms_regs_x86_fxsave_get_regs(struct SymsRegs *dest, SymsIntelFXSave *src);
+
+SYMS_INTERNAL void 
+syms_regs_x86_xsave_get_regs(struct SymsRegs *dest, SymsIntelXSave *src);
+
+SYMS_INTERNAL void 
+syms_regs_x64_fxsave_put_regs(SymsIntelFXSave *dest, struct SymsRegs *src);
+
+SYMS_INTERNAL void 
+syms_regs_x64_xsave_put_regs(SymsIntelXSave *dest, struct SymsRegs *src);
+
+SYMS_INTERNAL void 
+syms_regs_x86_fxsave_put_regs(SymsIntelFXSave *dest, struct SymsRegs *src);
+
+SYMS_INTERNAL void 
+syms_regs_x86_xsave_put_regs(SymsIntelXSave *dest, struct SymsRegs *src);
+
+#endif /* SYMS_REGS_INTEL_H */
