RAD symbol library

diff --git a/Engine/Source/Runtime/SymsLib/src/syms.c b/Engine/Source/Runtime/SymsLib/src/syms.c
new file mode 100644
--- /dev/null
+++ b/Engine/Source/Runtime/SymsLib/src/syms.c
@@ -0,0 +1,3443 @@
+// Copyright Epic Games, Inc. All Rights Reserved.
+/******************************************************************************
+ * File   : syms.c                                                            *
+ * Author : Nikita Smith                                                      *
+ * Created: 2019/02/20                                                        *
+ * Purpose: Implementation of user-level API                                  *
+ ******************************************************************************/
+
+SYMS_API SymsProc *
+syms_mod_info_find_proc(SymsModInfo *m, syms_uint i)
+{
+  return (SymsProc *)syms_block_allocator_find_push(m->procs, i);
+}
+
+SYMS_API SymsArch
+syms_get_arch(struct SymsInstance *instance)
+{
+  if (instance->img.arch != SYMS_ARCH_NULL) {
+    return instance->img.arch;
+  }
+  if (instance->debug_info.type == SYMS_FORMAT_PDB) {
+    SymsDebugInfoPdb *pdb_info = (SymsDebugInfoPdb *)instance->debug_info.impl;
+    SymsNTFileHeaderMachineType machine_type = pdb_get_machine_type(&pdb_info->context);
+    // TODO: move this code to syms_nt.h
+    switch (machine_type) {
+    case SYMS_NT_FILE_HEADER_MACHINE_X86: return SYMS_ARCH_X86;
+    case SYMS_NT_FILE_HEADER_MACHINE_X64: return SYMS_ARCH_X64;
+    case SYMS_NT_FILE_HEADER_MACHINE_ARM: return SYMS_ARCH_ARM32;
+    case SYMS_NT_FILE_HEADER_MACHINE_ARM64: return SYMS_ARCH_ARM;
+    default: break;
+    }
+  }
+  return SYMS_ARCH_NULL;
+}
+
+SYMS_INTERNAL SymsErrorCode
+syms_set_error(struct SymsInstance *instance, SymsErrorCode code)
+{
+  instance->error.code = code;
+  instance->error.text = syms_string_init(0, 0);
+  return code;
+}
+
+SYMS_INTERNAL SymsErrorCode
+syms_set_error_ex(struct SymsInstance *instance, SymsErrorCode code, const char *text)
+{
+  instance->error.code = code;
+  instance->error.text = syms_string_init_lit(text);
+  return code;
+}
+
+
+SYMS_INLINE SymsDebugInfoPdb *
+syms_get_debug_info_pdb(SymsDebugInfo *di)
+{
+  if (di->type == SYMS_FORMAT_PDB) {
+    return (SymsDebugInfoPdb *)di->impl;
+  }
+  return NULL;
+}
+
+SYMS_INLINE SymsDebugInfoDwarf *
+syms_get_debug_info_dwarf(SymsDebugInfo *di)
+{
+  if (di->type == SYMS_FORMAT_DWARF) {
+    return (SymsDebugInfoDwarf *)di->impl;
+  }
+  return NULL;
+}
+
+SYMS_INTERNAL SymsStringRef
+syms_string_ref_null(void)
+{
+  SymsStringRef result;
+  result.type = SYMS_STRING_REF_TYPE_NULL;
+  return result;
+}
+
+SYMS_API SymsStringRef
+syms_string_ref_str(SymsString str)
+{
+  SymsStringRef ref;
+  ref.type = SYMS_STRING_REF_TYPE_STR;
+  *(SymsString *)ref.impl = str;
+  return ref;
+}
+
+SYMS_API syms_bool
+syms_string_cmp_strref(struct SymsInstance *instance, SymsString a, SymsStringRef b)
+{
+  syms_bool is_equals = syms_false;
+  switch (b.type) {
+  case SYMS_STRING_REF_TYPE_NULL: {
+  } break;
+  case SYMS_STRING_REF_TYPE_PDB: {
+    SymsDebugInfoPdb *pdb_info = syms_get_debug_info_pdb(&instance->debug_info);
+    pdb_pointer *pdb_str = (pdb_pointer *)b.impl; 
+    is_equals = syms_string_cmp_strref_pdb(pdb_info, a, *pdb_str);
+  } break;
+  case SYMS_STRING_REF_TYPE_STR: {
+    SymsString *str = (SymsString *)b.impl;
+    is_equals = syms_string_cmp(a, *str);
+  } break;
+  case SYMS_STRING_REF_TYPE_DW_PATH: {
+    char temp[1024]; // TODO(nick): stack allocation
+    syms_uint temp_size = syms_read_strref(instance, &b, temp, sizeof(temp) - 1);
+    SymsString str = syms_string_init(temp, temp_size);
+    is_equals = syms_string_cmp(a, str);
+  } break;
+  }
+  return is_equals;
+}
+
+SYMS_INTERNAL syms_uint
+syms_read_strref_(struct SymsInstance *instance, SymsStringRef *ref, void *buffer, syms_uint buffer_max)
+{
+  syms_uint read_size = 0;
+  switch (ref->type) {
+  case SYMS_STRING_REF_TYPE_NULL: break;
+  case SYMS_STRING_REF_TYPE_PDB: {
+    SymsDebugInfoPdb *pdb_info = syms_get_debug_info_pdb(&instance->debug_info);
+    pdb_pointer *pdb_str = (pdb_pointer *)ref->impl;
+    read_size = syms_read_strref_pdb(pdb_info, pdb_str, buffer, pdb_trunc_uint(buffer_max));
+  } break;
+  case SYMS_STRING_REF_TYPE_STR: {
+    SymsBuffer temp = syms_buffer_init(buffer, buffer ? buffer_max : 4096);
+    SymsString *str = (SymsString *)ref->impl;
+    syms_buffer_write(&temp, str->data, pdb_trunc_uint(str->len));
+    syms_buffer_write_u8(&temp, '\0');
+    read_size = syms_trunc_u32(temp.off);
+  } break;
+  case SYMS_STRING_REF_TYPE_DW_PATH: {
+    SymsBuffer temp = syms_buffer_init(buffer, buffer ? buffer_max : 4096);
+    DwFilePath *path = (DwFilePath *)ref->impl;
+    SymsString *compile_dir = &path->compile_dir;
+    SymsString *dir = &path->dir;
+    SymsString *file = &path->file;
+
+    if (dir->len > 0) {
+      syms_bool is_windows = dir->len > 1 && dir->data[1] == ':';
+      syms_bool is_unix = dir->data[0] == '/';
+      if (!is_windows && !is_unix) {
+        syms_buffer_write(&temp, compile_dir->data, compile_dir->len);
+      }
+
+      // write directory
+      syms_buffer_write(&temp, dir->data, dir->len);
+
+      if (dir->data[dir->len - 1] != '\\' && dir->data[dir->len - 1] != '/') {
+        if (dir->data[0] != '/') {
+          // windows path
+          syms_buffer_write_u8(&temp, '\\');
+        } else {
+          // unix path
+          syms_buffer_write_u8(&temp, '/');
+        }
+      }
+    }
+
+    // write file name
+    syms_buffer_write(&temp, file->data, file->len);
+
+    // write null
+    syms_buffer_write_u8(&temp, '\0');
+
+    read_size = syms_trunc_u32(temp.off);
+  } break;
+  }
+  return read_size;
+}
+SYMS_API syms_uint
+syms_read_strref(SymsInstance *instance, SymsStringRef *ref, void *buffer, syms_uint buffer_size)
+{
+  syms_uint result = syms_read_strref_(instance, ref, buffer, buffer_size);
+  return result;
+}
+SYMS_API syms_uint
+syms_strref_get_size(SymsInstance *instance, SymsStringRef *ref)
+{
+  syms_uint result = syms_read_strref_(instance, ref, NULL, 0);
+  return result;
+}
+SYMS_API SymsString
+syms_read_strref_str(SymsInstance *instance, SymsStringRef *ref, SymsArena *arena)
+{
+  syms_uint buffer_size = syms_strref_get_size(instance, ref);
+  SymsString str = syms_string_init(0,0);
+  if (buffer_size > 0) {
+    char *buffer = syms_arena_push_array(arena, char, buffer_size);
+    if (buffer) {
+      syms_uint buffer_len = syms_read_strref(instance, ref, buffer, buffer_size);
+      str = syms_string_init(buffer, buffer_len);
+    }
+  }
+  return str;
+}
+
+SYMS_API char *
+syms_read_strref_arena(SymsInstance *instance, SymsStringRef *ref, SymsArena *arena)
+{
+  SymsString str = syms_read_strref_str(instance, ref, arena);
+  return (char *)str.data;
+}
+
+// -------------------------------------------------------------------------------- 
+
+SYMS_API SymsInstance *
+syms_init_ex(SymsAddr rebase)
+{
+  syms_init_os();
+  SymsArena *arena = syms_borrow_memory(0);
+  SymsInstance *instance = syms_arena_push_struct(arena, SymsInstance);
+  syms_memzero(instance, sizeof(*instance));
+  instance->rebase = rebase;
+  instance->arena = arena;
+  return instance;
+}
+
+SYMS_API SymsInstance *
+syms_init(void)
+{
+  return syms_init_ex(0);
+}
+
+SYMS_API void
+syms_quit(SymsInstance *instance)
+{
+  for (SymsBorrowedMemory *n = instance->borrowed_memory; n != 0; n = n->next) {
+    syms_return_memory(0, &n->arena);
+  }
+  SymsArena *arena = instance->arena;
+  syms_return_memory(0, &arena);
+}
+
+SYMS_INTERNAL SymsAddr
+syms_rebase_addr(const struct SymsInstance *instance, SymsAddr addr)
+{
+  SymsAddr orig_rebase = syms_get_orig_rebase(instance);
+  SymsAddr result = addr;
+  SymsAddr new_rebase = syms_get_rebase(instance);
+  if (instance->img.type == SYMS_IMAGE_ELF && instance->debug_info.type == SYMS_FORMAT_DWARF) {
+    if (addr >= orig_rebase) {
+      result = (addr - orig_rebase) + new_rebase;
+    }
+  } else {
+    result = addr + new_rebase;
+  }
+  return result;
+}
+
+// -------------------------------------------------------------------------------- 
+
+SYMS_API SymsErrorCode
+syms_load_image(struct SymsInstance *instance, void *image_base, SymsUMM image_length, SymsLoadImageFlags flags)
+{
+  SymsErrorCode err = SYMS_ERR_INVALID_CODE_PATH;
+
+  SymsImage *image = &instance->img;
+  syms_memzero(image, sizeof(*image));
+  image->img_data      = image_base;
+  image->img_data_size = image_length;
+  image->flags         = flags;
+  
+  if (image_base && image_length) {
+    do {
+      if (syms_img_init_nt(image, image_base, image_length, flags)) {
+        err = SYMS_ERR_OK;
+        break;
+      }
+      if (syms_img_init_elf(image, image_base, image_length, flags)) {
+        err = SYMS_ERR_OK;
+        break;
+      }
+    } while (0);
+  }
+
+  return err;
+}
+
+SYMS_API void
+syms_free_image(SymsInstance *instance)
+{ (void)instance;
+  // We aren't making allocations for the image
+}
+
+SYMS_API SymsImage *
+syms_get_image(SymsInstance *instance)
+{
+  return &instance->img;
+}
+
+SYMS_API SymsDebugInfo *
+syms_get_debug_info(SymsInstance *instance)
+{
+  return &instance->debug_info;
+}
+
+SYMS_API syms_bool
+syms_match_image_with_debug_info(const SymsImage *image, const SymsDebugInfo *debug_info)
+{
+  syms_bool is_matched = syms_false;
+
+  if (image->type == SYMS_IMAGE_NT && debug_info->type == SYMS_FORMAT_PDB) {
+    SymsImageNT *nt = (SymsImageNT*)image->impl;
+    SymsDebugInfoPdb *pdb = (SymsDebugInfoPdb*)debug_info->impl;
+    // NOTE: certain .dll and .pdb have mismatched ages, but even with that
+    // dbghelp.dll loads them anyways, same goes for visual studio. It is strange
+    // because MSVC 2019 writes correct age in image and pdb when building incrementally.
+#if 0
+    if (syms_memcmp(&nt->pdb_guid, &pdb->context.auth_guid, sizeof(SymsGUID)) == 0) {
+      is_matched = (nt->pdb_age == pdb->context.age);
+    }
+#else
+    is_matched = syms_memcmp(&nt->pdb_guid, &pdb->context.auth_guid, sizeof(SymsGUID)) == 0;
+#endif
+  } else if (image->type == SYMS_IMAGE_ELF && debug_info->type == SYMS_FORMAT_DWARF) {
+    // SYMS_ASSERT_FAILURE("not implemented");
+  }
+
+  return is_matched;
+}
+
+SYMS_API syms_bool
+syms_debug_file_iter_init(SymsDebugFileIter *iter_out, SymsInstance *instance)
+{
+  iter_out->file_index = 0;
+  iter_out->instance = instance;
+  return syms_true;
+}
+
+SYMS_API syms_bool
+syms_debug_file_iter_next(SymsDebugFileIter *iter, SymsString *path_out)
+{
+  SymsInstance *instance = iter->instance;
+  SymsImage *image = syms_get_image(instance);
+  syms_bool is_path_found = syms_false;
+  switch (image->type) {
+  case SYMS_IMAGE_NT: {
+    if (iter->file_index == 0) {
+      SymsImageNT *nt = (SymsImageNT *)image->impl;
+      *path_out = nt->pdb_path;
+      iter->file_index += 1;
+      is_path_found = syms_true;
+    }
+  } break;
+  case SYMS_IMAGE_ELF: {
+    SymsDebugFileIterDwarf *dw_iter = (SymsDebugFileIterDwarf*)iter->impl;
+    if (iter->file_index == 0) {
+      if (!syms_debug_file_iter_init_dwarf(instance, dw_iter)) {
+        return syms_false;
+      }
+    }
+    iter->file_index += 1;
+    is_path_found = syms_debug_file_iter_next_dwarf(dw_iter, path_out);
+  } break;
+  }
+  return is_path_found;
+}
+
+SYMS_API syms_uint
+syms_get_debug_file_count(SymsInstance *instance)
+{
+  syms_uint count = 0;
+  SymsDebugFileIter iter;
+  if (syms_debug_file_iter_init(&iter, instance)) {
+    SymsString dummy;
+    while (syms_debug_file_iter_next(&iter, &dummy)) {
+      count += 1;
+    }
+  }
+  return count;
+}
+
+// -------------------------------------------------------------------------------- 
+
+SYMS_INTERNAL syms_int
+syms_quicksort_part(SymsProc *procs, syms_int lo, syms_int hi)
+{
+  SymsAddr va = procs[hi].va;
+  syms_int pi = lo - 1;
+  syms_int i;
+  SymsProc t;
+
+  for (i = lo; i < hi; ++i) {
+    if (procs[i].va < va) {
+      ++pi;
+      t = procs[pi];
+      procs[pi] = procs[i];
+      procs[i] = t;
+    }
+  }
+
+  ++pi;
+  t = procs[pi];
+  procs[pi] = procs[hi];
+  procs[hi] = t;
+
+  return pi;
+}
+SYMS_INTERNAL void
+syms_quicksort_procs_(SymsProc *procs, syms_int lo, syms_int hi)
+{
+  if (lo < hi) {
+    syms_int pi = syms_quicksort_part(procs, lo, hi);
+    SYMS_ASSERT(pi > 0);
+    syms_quicksort_procs_(procs, lo, pi - 1);
+    syms_quicksort_procs_(procs, pi + 1, hi);
+  }
+}
+SYMS_INTERNAL void
+syms_quicksort_procs(SymsProc *proc, syms_uint count)
+{
+  if (count > 0) {
+    syms_quicksort_procs_(proc, 0, (syms_int)count - 1);
+  }
+}
+
+SYMS_INTERNAL int 
+syms_qsort_compar_procs(const void *a, const void *b)
+{
+  const SymsProc *proc_a = (const SymsProc *)a;
+  const SymsProc *proc_b = (const SymsProc *)b;
+  if (proc_a->va < proc_b->va) return -1;
+  if (proc_a->va > proc_b->va) return +1;
+#if defined(SYMS_PARANOID)
+  SYMS_ASSERT_FAILURE("address collision");
+#endif
+  return 0;
+}
+
+SYMS_INTERNAL int
+syms_qsort_rangemap(const void *a, const void *b)
+{
+  const SymsRangeMap *range_a = (const SymsRangeMap *)a;
+  const SymsRangeMap *range_b = (const SymsRangeMap *)b;
+  return range_a->lo < range_b->lo;
+}
+
+SYMS_API SymsErrorCode
+syms_load_debug_info_ex(SymsInstance *instance, SymsFile *files, syms_uint file_count, SymsLoadDebugInfoFlags load_flags)
+{
+  SymsImage *img_info = &instance->img;
+  SymsArena *arena = instance->arena;
+  SymsDebugInfo *debug_info = &instance->debug_info;
+  SymsErrorCode err = SYMS_ERR_INVALID_CODE_PATH;
+
+  syms_memzero(debug_info, sizeof(*debug_info));
+
+  // PDB
+  if (file_count > 0) {
+    SymsDebugInfoPdb *pdb_info;
+    void *file_base;
+    SymsUMM file_size;
+
+    pdb_info  = (SymsDebugInfoPdb *)debug_info->impl;
+    file_base = files[0].base;
+    file_size = files[0].size;
+
+    if (pdb_init(&pdb_info->context, file_base, pdb_trunc_uint(file_size))) {
+      SymsArenaFrame frame = syms_arena_frame_begin(arena);
+      if (pdb_load_aux_data(&pdb_info->context, arena)) {
+        debug_info->type = SYMS_FORMAT_PDB;
+      } else {
+        syms_arena_frame_end(frame);
+      }
+    }
+  }
+
+  // Dwarf
+  if (debug_info->type == SYMS_FORMAT_NULL) {
+    typedef struct 
+    {
+      SymsString name;
+      DwSecType index;
+    } SymsDwarfSectionMap;
+
+    DwInitdata init_data;
+    syms_uint i;
+    SymsDebugInfoDwarf *dw_info;
+
+    static SymsDwarfSectionMap g_dwarf_section_map[DW_SEC_MAX] = 
+    {
+      { SYMS_STRING_LIT(""), DW_SEC_NULL },
+      { SYMS_STRING_LIT(".debug_abbrev"), DW_SEC_ABBREV },
+      { SYMS_STRING_LIT(".debug_aranges"), DW_SEC_ARANGES },
+      { SYMS_STRING_LIT(".debug_frame"), DW_SEC_FRAME },
+      { SYMS_STRING_LIT(".debug_info"), DW_SEC_INFO },
+      { SYMS_STRING_LIT(".debug_line"), DW_SEC_LINE },
+      { SYMS_STRING_LIT(".debug_loc"), DW_SEC_LOC },
+      { SYMS_STRING_LIT(".debug_macinfo"), DW_SEC_MACINFO },
+      { SYMS_STRING_LIT(".debug_pubnames"), DW_SEC_PUBNAMES },
+      { SYMS_STRING_LIT(".debug_pubtypes"), DW_SEC_PUBTYPES },
+      { SYMS_STRING_LIT(".debug_ranges"), DW_SEC_RANGES },
+      { SYMS_STRING_LIT(".debug_str"), DW_SEC_STR },
+      { SYMS_STRING_LIT(".debug_addr"), DW_SEC_ADDR },
+      { SYMS_STRING_LIT(".debug_loclists"), DW_SEC_LOCLISTS },
+      { SYMS_STRING_LIT(".debug_rnglists"), DW_SEC_RNGLISTS },
+      { SYMS_STRING_LIT(".debug_str_offsets"), DW_SEC_STR_OFFSETS },
+      { SYMS_STRING_LIT(".debug_line_str"), DW_SEC_LINE_STR }
+    };
+
+    // scan image sections for dwarf info
+    for (i = 0; i < SYMS_ARRAY_SIZE(g_dwarf_section_map); ++i) {
+      DwSecType sec_index = g_dwarf_section_map[i].index;
+      SymsSection img_sec;
+      if (syms_img_sec_from_name(instance, g_dwarf_section_map[i].name, &img_sec)) {
+        init_data.secs[sec_index].data = img_sec.data;
+        init_data.secs[sec_index].data_len = img_sec.data_size;
+      } else {
+        init_data.secs[sec_index].data = 0;
+        init_data.secs[sec_index].data_len = 0;
+      }
+    }
+    // scan file list for dwarf info
+    for (i = 0; i < file_count; ++i) {
+      SymsString path = syms_string_init_lit(files[i].path);
+      SymsString file_name = syms_path_get_file_name(path, /* strip file extension */ syms_true);
+      syms_uint k;
+      for (k = 0; k < SYMS_ARRAY_SIZE(g_dwarf_section_map); ++k) {
+        if (syms_string_cmp(file_name, g_dwarf_section_map[k].name)) {
+          DwSecType sec_index = g_dwarf_section_map[k].index;
+          init_data.secs[sec_index].data = files[i].base;
+          init_data.secs[sec_index].data_len = files[i].size;
+        }
+      }
+    }
+
+    dw_info = (SymsDebugInfoDwarf *)debug_info->impl;
+    if (dw_init(&dw_info->context, img_info->arch, &init_data)) {
+      if (dw_load_heap(&dw_info->context, arena)) {
+        debug_info->type = SYMS_FORMAT_DWARF;
+      }
+    }
+  }
+
+  // ELF symbol table
+  if (debug_info->type == SYMS_FORMAT_NULL) {
+    if (img_info->type == SYMS_IMAGE_ELF) {
+      SymsSection dummy;
+      syms_bool has_symbol_table = syms_img_sec_from_name(instance, syms_string_lit(".symtab"), &dummy);
+      if (has_symbol_table) {
+        debug_info->type = SYMS_FORMAT_ELFSYMTAB;
+      }
+    }
+  }
+
+  if (debug_info->type != SYMS_FORMAT_NULL) {
+    SymsModIter mod_iter;
+    instance->mods_num = 0;
+    instance->mods     = 0;
+    if (syms_mod_iter_init(instance, &mod_iter)) {
+      // count modules
+      SymsMod mod;
+      syms_uint mod_index = 0;
+      while (syms_mod_iter_next(&mod_iter, &mod)) {
+        instance->mods_num += 1;
+      }
+      // reset iterator
+      if (!syms_mod_iter_init(instance, &mod_iter)) {
+        SYMS_INVALID_CODE_PATH;
+      }
+      // allocate memory
+      instance->mods = syms_arena_push_array(instance->arena, SymsModInfo, instance->mods_num);
+      SYMS_ASSERT(instance->mods);
+      // store modules
+      for (;;) {
+        SymsMod header;
+        SymsModInfo *info;
+
+        if (!syms_mod_iter_next(&mod_iter, &header)) {
+          break;
+        }
+        SYMS_ASSERT(mod_index < instance->mods_num);
+        info = &instance->mods[mod_index++];
+        info->header        = header;
+      }
+
+      // Load an address table for module searching.
+      // Addresses represent a low granuality range
+      // and you wont be able to resolve to exact
+      // symbol ( defered or lazy loading ).
+      // Both formats have simmilar solutions for this
+      // problem. Dwarf defines .debug_aranges section
+      // that is partitioned on a module basis.
+      // And PDB provides Section Contributions
+      // table; table is located in DBI stream right
+      // after module information.
+    }
+
+    if (~load_flags & SYMS_LOAD_DEBUG_INFO_FLAGS_DEFER_BUILD_MODULE) {
+      syms_uint module_count = syms_get_module_build_count(instance);
+      SymsModID mod_id;
+      for (mod_id = 0; mod_id < module_count; ++mod_id) {
+        syms_build_module(instance, mod_id, instance->arena);
+      }
+    }
+  }
+
+  if (debug_info->type != SYMS_FORMAT_NULL)
+    err = SYMS_ERR_OK;
+
+  return err;
+}
+
+// -------------------------------------------------------------------------------- 
+
+SYMS_API syms_bool
+syms_sec_iter_init(struct SymsInstance *instance, SymsSecIter *iter)
+{
+  iter->image = &instance->img;
+  switch (iter->image->type) {
+  case SYMS_IMAGE_NULL: break;
+  case SYMS_IMAGE_NT: {
+    SymsSecIterNT *iter_impl = (SymsSecIterNT *)iter->impl;
+    *iter_impl = syms_sec_iter_init_nt(iter->image);
+  } break;
+  case SYMS_IMAGE_ELF: {
+    SymsSecIterElf *iter_impl = (SymsSecIterElf *)iter->impl;
+    *iter_impl = syms_sec_iter_init_elf(iter->image); 
+  } break;
+  }
+  return syms_true;
+}
+
+SYMS_API syms_bool
+syms_sec_iter_next(SymsSecIter *iter, SymsSection *sec_out)
+{
+  SymsImage *image = iter->image;
+  syms_bool result = syms_false;
+  switch (image->type) {
+  case SYMS_IMAGE_NULL: break;
+  case SYMS_IMAGE_NT: {
+    SymsSecIterNT *iter_impl = (SymsSecIterNT *)iter->impl;
+    SymsNTImageSectionHeader *sec = (SymsNTImageSectionHeader *)sec_out->impl;
+    if (syms_sec_iter_next_nt(iter_impl, sec)) {
+      SymsAddr data_lo = (image->flags & SYMS_LOAD_IMAGE_FLAGS_FROM_MEMORY) ? sec->va : sec->rawdata_ptr;
+      SymsAddr data_hi = data_lo + sec->u.virtual_size; 
+
+      sec_out->name = syms_string_init_lit(&sec->name[0]);
+      sec_out->data = (void *)((u8 *)image->img_data + data_lo);
+      sec_out->data_size = data_hi - data_lo;
+      sec_out->off  = sec->rawdata_ptr;
+      sec_out->va   = sec->va;
+
+      result = syms_true;
+    }
+  } break;
+  case SYMS_IMAGE_ELF : {
+    SymsSecIterElf *iter_impl = (SymsSecIterElf *)iter->impl; 
+    SymsElfShdr64 *sec = (SymsElfShdr64 *)sec_out->impl;
+    if (syms_sec_iter_next_elf(iter_impl, sec)) {
+      SymsAddr sh_name_lo = iter_impl->sh_name_lo;
+      sec_out->name = syms_string_init_lit("");
+      if (sec->sh_name < sh_name_lo && sh_name_lo + sec->sh_name < image->img_data_size) {
+        SymsBuffer buffer = syms_buffer_init(image->img_data, image->img_data_size);
+        if (syms_buffer_seek(&buffer, sh_name_lo + sec->sh_name)) {
+          sec_out->name = syms_buffer_read_string(&buffer);
+        }
+      }
+
+      sec_out->data_size = sec->sh_size;
+      sec_out->off = sec->sh_offset;
+      sec_out->va = sec->sh_addr;
+
+      // check if section has data in file
+      sec_out->data = 0;
+      if (sec->sh_type != SYMS_SHT_NOBITS) {
+        SymsAddr data_lo = (image->flags & SYMS_LOAD_IMAGE_FLAGS_FROM_MEMORY) ? sec->sh_addr : sec->sh_offset;
+        SymsAddr data_hi = data_lo + sec->sh_size;
+        if (data_hi <= image->img_data_size) {
+          sec_out->data = (void *)((U8 *)image->img_data + data_lo);
+        }
+      }
+
+      result = syms_true;
+    }
+  } break;
+  }
+  return result;
+}
+
+SYMS_API SymsModInfo *
+syms_mod_from_va(struct SymsInstance *instance, SymsAddr va)
+{
+  SymsModID mod_id;
+  for (mod_id = 0; mod_id < instance->mods_num; ++mod_id) {
+    SymsModInfo *info = instance->mods + mod_id;
+    SymsRangeMap *map = syms_rangemap_search(info->rangemap, va);
+    if (map) {
+      return info;
+    }
+  }
+  return 0;
+}
+
+SYMS_API syms_bool
+syms_mod_iter_init(struct SymsInstance *instance, SymsModIter *iter)
+{
+  syms_bool result = syms_false;
+  SymsDebugInfo *debug_info = &instance->debug_info;
+
+  iter->instance = instance;
+  iter->index = 0;
+
+  switch (debug_info->type) {
+  case SYMS_FORMAT_NULL     : break;
+  case SYMS_FORMAT_ELFSYMTAB: result = syms_symtab_iter_init(instance, (SymsSymtabIter *)iter->impl);                                   break;
+  case SYMS_FORMAT_PDB      : result = syms_mod_iter_init_pdb(syms_get_debug_info_pdb(debug_info), (SymsModIterPdb *)iter->impl);       break;
+  case SYMS_FORMAT_DWARF    : result = syms_mod_iter_init_dwarf(syms_get_debug_info_dwarf(debug_info), (SymsModIterDwarf *)iter->impl); break;
+  }
+  return result;
+}
+SYMS_API syms_bool
+syms_mod_iter_next(SymsModIter *iter, SymsMod *mod_out)
+{
+  struct SymsInstance *instance = iter->instance;
+  SymsDebugInfo *debug_info = &instance->debug_info;
+  syms_bool result = syms_false;
+
+  switch (debug_info->type) {
+  case SYMS_FORMAT_NULL: break;
+  case SYMS_FORMAT_ELFSYMTAB: {
+    syms_memzero(mod_out, sizeof(*mod_out));
+    for (;;) {
+      SymsSymtabEntry entry;
+      if (!syms_symtab_iter_next((SymsSymtabIter *)iter, &entry)) {
+        break;
+      }
+      if (entry.type == SYMS_STT_FILE) {
+        mod_out->id = iter->index;
+        mod_out->name = syms_string_ref_str(entry.name);
+        mod_out->va = syms_get_rebase(instance);
+        mod_out->size = instance->img.img_data_size;
+        *(SymsSymtabIter *)mod_out->impl = *(SymsSymtabIter *)iter->impl;
+
+        result = syms_true;
+        break;
+      }
+    }
+  } break;
+  case SYMS_FORMAT_PDB  : result = syms_mod_iter_next_pdb(syms_get_debug_info_pdb(debug_info), (SymsModIterPdb *)iter->impl, mod_out);       break;
+  case SYMS_FORMAT_DWARF: result = syms_mod_iter_next_dwarf(syms_get_debug_info_dwarf(debug_info), (SymsModIterDwarf *)iter->impl, mod_out); break;
+  }
+  
+  iter->index += 1;
+
+  return result;
+}
+
+SYMS_INTERNAL void
+syms_line_init(SymsLine *line, SymsAddr va, syms_uint ln, syms_uint col)
+{
+  line->va  = va;
+  line->ln  = ln;
+  line->col = col;
+}
+SYMS_API syms_bool
+syms_line_iter_init(struct SymsInstance *instance, SymsLineIter *iter, SymsMod *mod)
+{
+  SymsDebugInfo *debug_info = &instance->debug_info;
+
+  syms_bool is_inited = syms_false;
+
+  iter->instance        = instance;
+  iter->debug_info      = debug_info;
+  iter->switched_file   = syms_false;
+  iter->has_line_count  = syms_false;
+  iter->line_count      = 0;
+
+  switch (debug_info->type) {
+  case SYMS_FORMAT_NULL : break;
+  case SYMS_FORMAT_PDB  : {
+    SymsDebugInfoPdb  *pdb_info   = syms_get_debug_info_pdb(debug_info);
+    SymsLineIterPdb   *pdb_iter   = (SymsLineIterPdb *)iter->impl;
+    pdb_mod           *pdb_module = (pdb_mod *)mod->impl;
+    is_inited = syms_line_iter_init_pdb(pdb_info, pdb_iter, pdb_module, &iter->has_line_count);
+  } break;
+  case SYMS_FORMAT_DWARF: {
+    SymsDebugInfoDwarf  *dwarf_info = syms_get_debug_info_dwarf(debug_info);
+    SymsLineIterDwarf   *dwarf_iter = (SymsLineIterDwarf *)iter->impl;
+    DwCompileUnit       *dwarf_mod  = (DwCompileUnit *)mod->impl;
+    is_inited = syms_line_iter_init_dwarf(dwarf_info, dwarf_iter, dwarf_mod, &iter->has_line_count); 
+  } break;
+  }
+
+  return is_inited;
+}
+SYMS_API syms_bool
+syms_line_iter_next(SymsLineIter *iter, SymsLine *line_out)
+{
+  SymsDebugInfo *debug_info = iter->debug_info;
+  syms_bool result = syms_false;
+  
+  iter->switched_file = syms_false;
+
+  switch (debug_info->type) {
+  case SYMS_FORMAT_NULL : break;
+  case SYMS_FORMAT_PDB  : {
+    SymsDebugInfoPdb *info_impl = syms_get_debug_info_pdb(debug_info);
+    SymsLineIterPdb *iter_impl = syms_line_iter_to_pdb(iter);
+    result = syms_line_iter_next_pdb(info_impl, iter_impl, 
+        line_out, &iter->switched_file, &iter->file, &iter->line_count);
+  } break;
+  case SYMS_FORMAT_DWARF: {
+    SymsDebugInfoDwarf *info_impl = syms_get_debug_info_dwarf(debug_info);
+    SymsLineIterDwarf *iter_impl = syms_line_iter_to_dw(iter);
+    result = syms_line_iter_next_dwarf(info_impl, iter_impl, line_out, &iter->switched_file, &iter->file);
+  } break;
+  }
+  line_out->va = syms_rebase_addr(iter->instance, line_out->va);
+
+  return result;
+}
+
+SYMS_API syms_bool
+syms_member_iter_init(struct SymsInstance *instance, SymsMemberIter *iter, SymsType *type)
+{
+  syms_bool result = syms_false;
+  SymsDebugInfo *debug_info = &instance->debug_info;
+  SymsType base_type = syms_infer_base_type(instance, type);
+  iter->debug_info = debug_info;
+  switch (debug_info->type) {
+  case SYMS_FORMAT_NULL : break;
+  case SYMS_FORMAT_PDB  : {
+    SymsDebugInfoPdb *info_impl = syms_get_debug_info_pdb(debug_info);
+    SymsMemberIterPdb *iter_impl = (SymsMemberIterPdb *)iter->impl;
+    result = syms_member_iter_init_pdb(info_impl, iter_impl, &base_type);
+  } break;
+  case SYMS_FORMAT_DWARF: {
+    SymsDebugInfoDwarf *info_impl = syms_get_debug_info_dwarf(debug_info);
+    SymsMemberIterDwarf *iter_impl = (SymsMemberIterDwarf *)iter->impl;
+    result = syms_member_iter_init_dwarf(info_impl, iter_impl, &base_type); 
+  } break;
+  }
+  return result;
+}
+
+SYMS_API syms_bool
+syms_member_iter_next(SymsMemberIter *iter, SymsMember *member_out)
+{
+  SymsDebugInfo *debug_info = iter->debug_info;
+  syms_bool result = syms_false;
+  switch (debug_info->type) {
+  case SYMS_FORMAT_NULL : break;
+  case SYMS_FORMAT_PDB  : {
+    SymsDebugInfoPdb *info_impl = syms_get_debug_info_pdb(debug_info);
+    SymsMemberIterPdb *iter_impl = (SymsMemberIterPdb *)iter->impl;
+    result = syms_member_iter_next_pdb(info_impl, iter_impl, member_out);
+  } break;
+  case SYMS_FORMAT_DWARF: {
+    SymsDebugInfoDwarf *info_impl = syms_get_debug_info_dwarf(debug_info);
+    SymsMemberIterDwarf *iter_impl = (SymsMemberIterDwarf *)iter->impl;
+    result = syms_member_iter_next_dwarf(info_impl, iter_impl, member_out); 
+  } break;
+  }
+  return result;
+}
+
+SYMS_API syms_bool
+syms_global_iter_init(struct SymsInstance *instance, SymsGlobalIter *iter)
+{
+  syms_bool result = syms_false;
+  SymsDebugInfo *debug_info = &instance->debug_info;
+  iter->debug_info = debug_info;
+  switch (debug_info->type) {
+  case SYMS_FORMAT_NULL: break;
+  case SYMS_FORMAT_PDB: {
+    SymsDebugInfoPdb *info_impl = syms_get_debug_info_pdb(debug_info);
+    SymsGlobalIterPdb *iter_impl = (SymsGlobalIterPdb *)iter->impl;
+    result = syms_global_iter_init_pdb(info_impl, iter_impl);
+  } break;
+  case SYMS_FORMAT_DWARF: {
+    SymsDebugInfoDwarf *info_impl = syms_get_debug_info_dwarf(debug_info);
+    SymsGlobalIterDwarf *iter_impl = (SymsGlobalIterDwarf *)iter->impl;
+    result = syms_global_iter_init_dwarf(info_impl, iter_impl);
+  } break;
+  }
+  return result;
+}
+
+SYMS_API syms_bool
+syms_global_iter_next(SymsGlobalIter *iter, SymsGlobal *gdata_out)
+{
+  syms_bool result = syms_false;
+  SymsDebugInfo *debug_info = iter->debug_info;
+  switch (debug_info->type) {
+  case SYMS_FORMAT_NULL: break;
+  case SYMS_FORMAT_PDB: {
+    SymsDebugInfoPdb *info_impl = syms_get_debug_info_pdb(debug_info);
+    SymsGlobalIterPdb *iter_impl = (SymsGlobalIterPdb *)iter->impl;
+    result = syms_global_iter_next_pdb(info_impl, iter_impl, gdata_out);       
+  } break;
+  case SYMS_FORMAT_DWARF: {
+    SymsDebugInfoDwarf *info_impl = syms_get_debug_info_dwarf(debug_info);
+    SymsGlobalIterDwarf *iter_impl = (SymsGlobalIterDwarf *)iter->impl;
+    result = syms_global_iter_next_dwarf(info_impl, iter_impl, gdata_out); 
+  } break;
+  }
+  return result;
+}
+
+SYMS_API syms_bool
+syms_local_data_iter_init(struct SymsInstance *instance, SymsLocalDataIter *iter, SymsMod *mod)
+{
+  SymsDebugInfo *debug_info = &instance->debug_info;
+  syms_bool result = syms_false;
+  iter->debug_info = debug_info;
+  iter->mod_id = mod->id;
+  switch (debug_info->type) {
+  case SYMS_FORMAT_NULL: break;
+  case SYMS_FORMAT_PDB: {
+    SymsDebugInfoPdb *info_impl = syms_get_debug_info_pdb(debug_info);
+    SymsLocalDataIterPdb *iter_impl = (SymsLocalDataIterPdb *)iter->impl;
+    result = syms_local_data_iter_init_pdb(info_impl, iter_impl, mod);
+  } break;
+  case SYMS_FORMAT_DWARF: {
+    SymsDebugInfoDwarf *info_impl = syms_get_debug_info_dwarf(debug_info);
+    SymsLocalDataIterDwarf *iter_impl = (SymsLocalDataIterDwarf *)iter->impl;
+    result = syms_local_data_iter_init_dwarf(info_impl, iter_impl, mod);  
+  } break;
+  }
+  return result;
+}
+SYMS_API syms_bool
+syms_local_data_iter_next(SymsLocalDataIter *iter, SymsLocalData *ldata_out)
+{
+  SymsDebugInfo *debug_info = iter->debug_info;
+  syms_bool result = syms_false;
+  switch (debug_info->type) {
+  case SYMS_FORMAT_NULL: break;
+  case SYMS_FORMAT_PDB: {
+    SymsDebugInfoPdb *info_impl = syms_get_debug_info_pdb(debug_info);
+    SymsLocalDataIterPdb *iter_impl = (SymsLocalDataIterPdb *)iter->impl;
+    result = syms_local_data_iter_next_pdb(info_impl, iter_impl, ldata_out);
+  } break;
+  case SYMS_FORMAT_DWARF: {
+    SymsDebugInfoDwarf *info_impl = syms_get_debug_info_dwarf(debug_info);
+    SymsLocalDataIterDwarf *iter_impl = (SymsLocalDataIterDwarf *)iter->impl;
+    result = syms_local_data_iter_next_dwarf(info_impl, iter_impl, ldata_out); 
+  } break;
+  }
+  return result;
+}
+
+SYMS_API syms_bool
+syms_proc_iter_init(struct SymsInstance *instance, SymsProcIter *iter, SymsMod *mod)
+{
+  syms_bool result = syms_false;
+  SymsDebugInfo *debug_info = &instance->debug_info;
+  iter->instance = instance;
+
+  switch (debug_info->type) {
+  case SYMS_FORMAT_NULL: break;
+  case SYMS_FORMAT_PDB: {
+    SymsDebugInfoPdb *info_impl = syms_get_debug_info_pdb(debug_info);
+    SymsProcIterPdb *iter_impl = (SymsProcIterPdb *)iter->impl;
+    result = syms_proc_iter_init_pdb(info_impl, iter_impl, mod);
+  } break;
+  case SYMS_FORMAT_DWARF: {
+    SymsDebugInfoDwarf *info_impl = syms_get_debug_info_dwarf(debug_info);
+    SymsProcIterDwarf *iter_impl = (SymsProcIterDwarf *)iter->impl;
+    result = syms_proc_iter_init_dwarf(info_impl, iter_impl, mod); 
+  } break;
+  case SYMS_FORMAT_ELFSYMTAB: {
+    *(SymsSymtabIter *)iter->impl = *(SymsSymtabIter *)mod->impl;
+    result = syms_true;
+  } break;
+  }
+  return result;
+}
+
+SYMS_API syms_bool
+syms_proc_iter_next(SymsProcIter *iter, SymsProc *proc_out)
+{
+  syms_bool result = syms_false;
+  struct SymsInstance  *instance   = iter->instance;
+  SymsDebugInfo *debug_info = &instance->debug_info;
+
+  switch (debug_info->type) {
+  case SYMS_FORMAT_NULL: break;
+  case SYMS_FORMAT_PDB: {
+    SymsDebugInfoPdb *info_impl = syms_get_debug_info_pdb(debug_info);
+    SymsProcIterPdb *iter_impl = (SymsProcIterPdb *)iter->impl;
+    result = syms_proc_iter_next_pdb(info_impl, iter_impl, proc_out);
+  } break;
+  case SYMS_FORMAT_DWARF: {
+    SymsDebugInfoDwarf *info_impl = syms_get_debug_info_dwarf(debug_info);
+    SymsProcIterDwarf *iter_impl = (SymsProcIterDwarf *)iter->impl;
+    result = syms_proc_iter_next_dwarf(info_impl, iter_impl, proc_out); 
+  } break;
+  case SYMS_FORMAT_ELFSYMTAB: {
+    SymsSymtabIter *iter_impl = (SymsSymtabIter *)iter->impl;
+    for (;;) {
+      SymsSymtabEntry entry;
+      if (!syms_symtab_iter_next(iter_impl, &entry)) {
+        break;
+      }
+      if (entry.type == SYMS_STT_FUNC) {
+        *proc_out = syms_proc_from_stt_func(&entry);
+        result = syms_true;
+        break;
+      } else if (entry.type == SYMS_STT_FILE) {
+        result = syms_false;
+        break;
+      }
+    }
+  } break;
+  }
+
+  proc_out->va           = syms_rebase_addr(instance, proc_out->va);
+  proc_out->dbg_start_va = syms_rebase_addr(instance, proc_out->dbg_start_va);
+  proc_out->dbg_end_va   = syms_rebase_addr(instance, proc_out->dbg_end_va);
+
+  return result;
+}
+
+SYMS_API syms_bool
+syms_local_iter_init(struct SymsInstance *instance, SymsLocalIter *iter, SymsProc *proc, void *stack, syms_uint stack_size)
+{
+  syms_bool result = syms_false;
+  SymsDebugInfo *debug_info = &instance->debug_info;
+  iter->instance = instance;
+  switch (debug_info->type) {
+  case SYMS_FORMAT_NULL: break;
+  case SYMS_FORMAT_PDB: {
+    SymsDebugInfoPdb *info_impl = syms_get_debug_info_pdb(debug_info);
+    SymsLocalIterPdb *iter_impl = (SymsLocalIterPdb *)iter->impl;
+    result = syms_local_iter_init_pdb(info_impl, iter_impl, proc, stack, stack_size);
+  } break;
+  case SYMS_FORMAT_DWARF: {
+    SymsDebugInfoDwarf *info_impl = syms_get_debug_info_dwarf(debug_info);
+    SymsLocalIterDwarf *iter_impl = (SymsLocalIterDwarf *)iter->impl;
+    result = syms_local_iter_init_dwarf(info_impl, iter_impl, proc); 
+  } break;
+  }
+  return result;
+}
+
+SYMS_API syms_bool
+syms_local_iter_next(SymsLocalIter *iter, SymsLocalExport *export_out)
+{
+  syms_bool result = syms_false;
+  SymsInstance *instance = iter->instance;
+  SymsDebugInfo *debug_info = &instance->debug_info;
+  switch (debug_info->type) {
+  case SYMS_FORMAT_NULL: break;
+  case SYMS_FORMAT_PDB: {
+    SymsDebugInfoPdb *info_impl = syms_get_debug_info_pdb(debug_info);
+    SymsLocalIterPdb *iter_impl = (SymsLocalIterPdb *)iter->impl;
+    SymsAddr rebase = syms_get_rebase(instance);
+    result = syms_local_iter_next_pdb(info_impl, iter_impl, rebase, export_out);
+  } break;
+  case SYMS_FORMAT_DWARF: {
+    SymsDebugInfoDwarf *info_impl = syms_get_debug_info_dwarf(debug_info);
+    SymsLocalIterDwarf *iter_impl = (SymsLocalIterDwarf *)iter->impl;
+    SymsVar var;
+    result = syms_local_iter_next_dwarf(info_impl, iter_impl, &var); 
+    export_out->type = SYMS_LOCAL_EXPORT_VAR;
+    export_out->u.var = var;
+  } break;
+  }
+  return result;
+}
+
+SYMS_API syms_bool
+syms_file_iter_init(struct SymsInstance *instance, SymsFileIter *iter)
+{
+  SymsDebugInfo *debug_info = &instance->debug_info;
+  syms_bool result = syms_false;
+  iter->debug_info = debug_info;
+  switch (debug_info->type) {
+  case SYMS_FORMAT_NULL: break;
+  case SYMS_FORMAT_PDB: {
+    SymsDebugInfoPdb *info_impl = syms_get_debug_info_pdb(debug_info);
+    SymsFileIterPdb *iter_impl = (SymsFileIterPdb *)iter->impl;
+    result = SymsFile_iter_init_pdb(info_impl, iter_impl);
+  } break;
+  case SYMS_FORMAT_DWARF: { 
+    SymsDebugInfoDwarf *info_impl = syms_get_debug_info_dwarf(debug_info);
+    SymsFileIterDwarf *iter_impl = (SymsFileIterDwarf *)iter->impl;
+    result = SymsFile_iter_init_dwarf(info_impl, iter_impl); 
+  } break;
+  }
+  return result;
+}
+
+SYMS_API syms_bool
+syms_file_iter_next(SymsFileIter *iter, SymsStringRef *ref_out)
+{
+  syms_bool result = syms_false;
+  SymsDebugInfo *debug_info = iter->debug_info;
+  switch (debug_info->type) {
+  case SYMS_FORMAT_NULL: break;
+  case SYMS_FORMAT_PDB: {
+    SymsDebugInfoPdb *info_impl = syms_get_debug_info_pdb(debug_info);
+    SymsFileIterPdb *iter_impl = (SymsFileIterPdb *)iter->impl;
+    result = SymsFile_iter_next_pdb(info_impl, iter_impl, ref_out);
+  } break;
+  case SYMS_FORMAT_DWARF: {
+    SymsDebugInfoDwarf *info_impl = syms_get_debug_info_dwarf(debug_info);
+    SymsFileIterDwarf *iter_impl = (SymsFileIterDwarf *)iter->impl;
+    result = SymsFile_iter_next_dwarf(info_impl, iter_impl, ref_out); 
+  } break;
+  }
+  return result;
+}
+
+SYMS_INTERNAL syms_bool
+syms_arg_iter_init_(struct SymsInstance *instance, SymsArgIter *iter, SymsTypeID *id)
+{
+  SymsDebugInfo *debug_info = &instance->debug_info;
+  syms_bool result = syms_false;
+  iter->debug_info = debug_info;
+  switch (debug_info->type) {
+  case SYMS_FORMAT_NULL: break;
+  case SYMS_FORMAT_PDB: {
+    SymsDebugInfoPdb *info_impl = syms_get_debug_info_pdb(debug_info);
+    SymsArgIterPdb *iter_impl = (SymsArgIterPdb *)iter->impl;
+    result = syms_arg_iter_init_pdb(info_impl, iter_impl, id);
+  } break;
+  case SYMS_FORMAT_DWARF: {
+    SymsDebugInfoDwarf *info_impl = syms_get_debug_info_dwarf(debug_info);
+    SymsArgIterDwarf *iter_impl = (SymsArgIterDwarf *)iter->impl;
+    result = syms_arg_iter_init_dwarf(info_impl, iter_impl, id); 
+  } break;
+  }
+  return result;
+}
+
+SYMS_API syms_bool
+syms_arg_iter_init(struct SymsInstance *instance, SymsArgIter *iter, SymsType *proc_type)
+{
+  syms_bool result;
+  switch (proc_type->kind) {
+  default: result = syms_false; break;
+  case SYMS_TYPE_PROC  : result = syms_arg_iter_init_(instance, iter, &proc_type->u.proc.arglist_type_id);   break;
+  case SYMS_TYPE_METHOD: result = syms_arg_iter_init_(instance, iter, &proc_type->u.method.arglist_type_id); break;
+  }
+  return result;
+}
+
+SYMS_API syms_bool
+syms_arg_iter_next(SymsArgIter *iter, SymsTypeID *typeid_out)
+{
+  SymsDebugInfo *debug_info = iter->debug_info;
+  syms_bool result = syms_false;
+  switch (debug_info->type) {
+  case SYMS_FORMAT_NULL : break;
+  case SYMS_FORMAT_PDB: {
+    SymsDebugInfoPdb *info_impl = syms_get_debug_info_pdb(debug_info);
+    SymsArgIterPdb *iter_impl = (SymsArgIterPdb *)iter->impl;
+    result = syms_arg_iter_next_pdb(info_impl, iter_impl, typeid_out);
+  } break;
+  case SYMS_FORMAT_DWARF: {
+    SymsDebugInfoDwarf *info_impl = syms_get_debug_info_dwarf(debug_info);
+    SymsArgIterDwarf *iter_impl = (SymsArgIterDwarf *)iter->impl;
+    result = syms_arg_iter_next_dwarf(info_impl, iter_impl, typeid_out); 
+  } break;
+  }
+  return result;
+}
+
+SYMS_API syms_bool
+syms_inline_iter_init(struct SymsInstance *instance, SymsInlineIter *iter, SymsAddr pc)
+{
+  syms_bool is_inited = syms_false;
+  SymsAddr rebase = syms_get_rebase(instance);
+  SymsDebugInfo *debug_info = &instance->debug_info;
+
+  iter->instance   = instance;
+  iter->debug_info = debug_info;
+
+  if (pc >= rebase) {
+    SymsProc proc;
+    SymsModID mod_id;
+    if (syms_proc_from_va_2(instance, pc, &proc, &mod_id)) {
+      SymsAddr rva = pc - rebase;
+      switch (debug_info->type) {
+      case SYMS_FORMAT_NULL : break;
+      case SYMS_FORMAT_PDB  : {
+        SymsDebugInfoPdb *info_impl = syms_get_debug_info_pdb(debug_info);
+        SymsInlineIterPdb *iter_impl = (SymsInlineIterPdb *)iter->impl;
+        is_inited = syms_inline_iter_init_pdb(info_impl, iter_impl, &proc, rva); 
+      } break;
+      case SYMS_FORMAT_DWARF: {
+        SymsModInfo *mod = syms_get_mod(instance, mod_id);
+
+        SymsDebugInfoDwarf  *dwarf_info = syms_get_debug_info_dwarf(debug_info);
+        SymsInlineIterDwarf *dwarf_iter = (SymsInlineIterDwarf *)iter->impl;
+        DwTag                dwarf_tag  = ((SymsProcData *)proc.impl)->dw.type_tag;
+        DwCompileUnit       *dwarf_mod  = (DwCompileUnit *)mod->header.impl;
+
+        is_inited = syms_inline_iter_init_dwarf(dwarf_info, dwarf_iter, dwarf_mod, dwarf_tag, rva);
+      } break;
+      }
+    }
+  }
+
+  return is_inited;
+}
+
+SYMS_API syms_bool
+syms_inline_iter_next(SymsInlineIter *iter, SymsInlineSite *site_out)
+{
+  syms_bool result = syms_false;
+  SymsDebugInfo *debug_info = iter->debug_info;
+  switch (debug_info->type) {
+  case SYMS_FORMAT_NULL: break;
+  case SYMS_FORMAT_PDB: {
+    SymsDebugInfoPdb *info_impl = syms_get_debug_info_pdb(debug_info);
+    SymsInlineIterPdb *iter_impl = (SymsInlineIterPdb *)iter->impl;
+    result = syms_inline_iter_next_pdb(info_impl, iter_impl, site_out); 
+  } break;
+  case SYMS_FORMAT_DWARF: {
+    SymsDebugInfoDwarf *info_impl = syms_get_debug_info_dwarf(debug_info);
+    SymsInlineIterDwarf *iter_impl = (SymsInlineIterDwarf *)iter->impl;
+    result = syms_inline_iter_next_dwarf(info_impl, iter_impl, site_out);
+  } break;
+  }
+  site_out->range_lo = syms_rebase_addr(iter->instance, site_out->range_lo);
+  site_out->range_hi = syms_rebase_addr(iter->instance, site_out->range_hi);
+  return result;
+}
+
+SYMS_API syms_bool
+syms_const_iter_init(struct SymsInstance *instance, SymsConstIter *iter)
+{
+  SymsDebugInfo *debug_info = &instance->debug_info;
+  syms_bool result = syms_false;
+  iter->debug_info = debug_info;
+  switch (debug_info->type) {
+  case SYMS_FORMAT_NULL: break;
+  case SYMS_FORMAT_PDB: {
+    SymsDebugInfoPdb *info_impl = syms_get_debug_info_pdb(debug_info);
+    SymsConstIterPdb *iter_impl = (SymsConstIterPdb *)iter->impl;
+    result = syms_const_iter_init_pdb(info_impl, iter_impl);
+  } break;
+  case SYMS_FORMAT_DWARF: {
+    SymsDebugInfoDwarf *info_impl = syms_get_debug_info_dwarf(debug_info);
+    SymsConstIterDwarf *iter_impl = (SymsConstIterDwarf *)iter->impl;
+    result = syms_const_iter_init_dwarf(info_impl, iter_impl);
+  } break;
+  }
+  return result;
+}
+
+SYMS_API syms_bool
+syms_const_iter_next(SymsConstIter *iter, SymsConst *const_out)
+{
+  syms_bool result = syms_false;
+  SymsDebugInfo *debug_info = iter->debug_info;
+  switch (debug_info->type) {
+  case SYMS_FORMAT_NULL: break;
+  case SYMS_FORMAT_PDB: {
+    SymsDebugInfoPdb *info_impl = syms_get_debug_info_pdb(debug_info);
+    SymsConstIterPdb *iter_impl = (SymsConstIterPdb *)iter->impl;
+    result = syms_const_iter_next_pdb(info_impl, iter_impl, const_out);
+  } break;
+  case SYMS_FORMAT_DWARF: {
+    SymsDebugInfoDwarf *info_impl = syms_get_debug_info_dwarf(debug_info);
+    SymsConstIterDwarf *iter_impl = (SymsConstIterDwarf *)iter->impl;
+    result = syms_const_iter_next_dwarf(info_impl, iter_impl, const_out); 
+  } break;
+  }
+  return result;
+}
+
+SYMS_API syms_bool
+syms_type_from_name(struct SymsInstance *instance, const char *name, syms_uint name_size, SymsType *type_out)
+{
+  SymsDebugInfo *debug_info = &instance->debug_info;
+  syms_bool result = syms_false;
+  SymsString name_str = syms_string_init(name, name_size);
+  switch (debug_info->type) {
+  case SYMS_FORMAT_NULL : break;
+  case SYMS_FORMAT_PDB  : result = syms_type_from_name_pdb(syms_get_debug_info_pdb(debug_info), name_str, type_out);     break;
+  case SYMS_FORMAT_DWARF: result = syms_type_from_name_dwarf(syms_get_debug_info_dwarf(debug_info), name_str, type_out); break;
+  }
+  return result;
+}
+
+SYMS_API syms_bool
+syms_const_from_name(struct SymsInstance *instance, const char *name, syms_uint name_size, SymsConst *const_out)
+{
+  SymsDebugInfo *debug_info = &instance->debug_info;
+  SymsString name_str = syms_string_init(name, name_size);
+  syms_bool result = syms_false;
+  switch (debug_info->type) {
+  case SYMS_FORMAT_NULL : break;
+  case SYMS_FORMAT_PDB  : result = syms_const_from_name_pdb(syms_get_debug_info_pdb(debug_info), name_str, const_out);     break;
+  case SYMS_FORMAT_DWARF: result = syms_const_from_name_dwarf(syms_get_debug_info_dwarf(debug_info), name_str, const_out); break;
+  }
+  return result;
+}
+
+SYMS_API syms_bool
+syms_global_from_name(struct SymsInstance *instance, const char *name, syms_uint name_size, SymsGlobal *gvar_out)
+{
+  SymsDebugInfo *debug_info = &instance->debug_info;
+  syms_bool result = syms_false;
+  SymsString name_str = syms_string_init(name, name_size);
+  switch (debug_info->type) {
+  case SYMS_FORMAT_NULL : break;
+  case SYMS_FORMAT_PDB  : result = syms_global_from_name_pdb(syms_get_debug_info_pdb(debug_info), name_str, gvar_out);     break;
+  case SYMS_FORMAT_DWARF: result = syms_global_from_name_dwarf(syms_get_debug_info_dwarf(debug_info), name_str, gvar_out); break;
+  }
+  return result;
+}
+
+SYMS_API syms_bool
+syms_type_iter_init(struct SymsInstance *instance, SymsTypeIter *iter)
+{
+  syms_bool result = syms_false;
+  SymsDebugInfo *debug_info = &instance->debug_info;
+  iter->debug_info = debug_info;
+  switch (debug_info->type) {
+  case SYMS_FORMAT_NULL: break;
+  case SYMS_FORMAT_PDB: {
+    SymsDebugInfoPdb *info_impl = syms_get_debug_info_pdb(debug_info);
+    SymsTypeIterPdb *iter_impl = (SymsTypeIterPdb *)iter->impl;
+    result = syms_type_iter_init_pdb(info_impl, iter_impl);
+  } break;
+  case SYMS_FORMAT_DWARF: {
+    SymsDebugInfoDwarf *info_impl = syms_get_debug_info_dwarf(debug_info);
+    SymsTypeIterDwarf *iter_impl = (SymsTypeIterDwarf *)iter->impl;
+    result = syms_type_iter_init_dwarf(info_impl, iter_impl);
+  } break;
+  }
+  return result;
+}
+
+SYMS_API syms_bool
+syms_type_iter_next(SymsTypeIter *iter, SymsTypeID *id_out)
+{
+  syms_bool result = syms_false;
+  SymsDebugInfo *debug_info = iter->debug_info;
+  switch (debug_info->type) {
+  case SYMS_FORMAT_NULL: break;
+  case SYMS_FORMAT_PDB: {
+    SymsDebugInfoPdb *info_impl = syms_get_debug_info_pdb(debug_info);
+    SymsTypeIterPdb *iter_impl = (SymsTypeIterPdb *)iter->impl;
+    result = syms_type_iter_next_pdb(info_impl, iter_impl, id_out);
+  } break;
+  case SYMS_FORMAT_DWARF: {
+    SymsDebugInfoDwarf *info_impl = syms_get_debug_info_dwarf(debug_info);
+    SymsTypeIterDwarf *iter_impl = (SymsTypeIterDwarf *)iter->impl;
+    result = syms_type_iter_next_dwarf(info_impl, iter_impl, id_out); 
+  } break;
+  }
+  return result;
+}
+
+SYMS_API syms_bool
+syms_range_iter_init(SymsRangeIter *iter, struct SymsInstance *instance, SymsRange *range)
+{
+  SymsDebugInfo *debug_info = &instance->debug_info;
+  syms_bool is_inited = syms_false;
+
+  iter->instance   = instance;
+  iter->debug_info = debug_info;
+  iter->range      = *range;
+
+  switch (range->type) {
+  case SYMS_RANGE_NULL: break;
+  case SYMS_RANGE_PLAIN: {
+    is_inited = syms_true;
+  } break;
+  case SYMS_RANGE_IMPL: {
+    switch (debug_info->type) {
+    case SYMS_FORMAT_NULL:
+    case SYMS_FORMAT_ELFSYMTAB:
+
+    case SYMS_FORMAT_PDB: { 
+      SymsDebugInfoPdb *info_impl  = (SymsDebugInfoPdb *)debug_info->impl;
+      SymsRangeIterPdb *iter_impl  = (SymsRangeIterPdb *)iter->impl;
+      SymsRangePdb     *range_impl = (SymsRangePdb *)range->u.impl;
+      is_inited = syms_range_iter_init_pdb(info_impl, iter_impl, range_impl->lo, range_impl->hi, range_impl->gaps);
+    } break;
+
+    case SYMS_FORMAT_DWARF: {
+      SymsRangeDwarf *range_impl = (SymsRangeDwarf *)range->u.impl;
+      SymsModInfo    *info       = syms_get_mod(instance, range_impl->cu_index);
+      SymsMod        *mod        = &info->header;
+      DwCompileUnit  *mod_impl   = (DwCompileUnit *)mod->impl;
+      DwRangeIter    *iter_impl  = (DwRangeIter *)iter->impl;
+      is_inited = dw_range_iter_init(iter_impl, mod_impl, range_impl->range_off);
+    } break;
+    }
+  } break;
+  }
+
+  return is_inited;
+}
+
+SYMS_API syms_bool
+syms_range_iter_next(SymsRangeIter *iter, SymsAddr *lo_out, SymsAddr *hi_out)
+{
+  syms_bool is_next_valid = syms_false;
+
+  switch (iter->range.type) {
+  case SYMS_RANGE_NULL: break;
+  case SYMS_RANGE_PLAIN: {
+    *lo_out = iter->range.u.plain.lo;
+    *hi_out = iter->range.u.plain.hi;
+    iter->range.type = SYMS_RANGE_NULL;
+    is_next_valid = syms_true;
+  } break;
+  case SYMS_RANGE_IMPL: {
+    SymsDebugInfo *debug_info = iter->debug_info;
+    switch (debug_info->type) {
+    case SYMS_FORMAT_NULL:
+    case SYMS_FORMAT_ELFSYMTAB:
+    case SYMS_FORMAT_PDB: {
+      SymsRangeIterPdb *iter_impl = (SymsRangeIterPdb *)iter->impl;
+      is_next_valid = syms_range_iter_next_pdb(iter_impl, lo_out, hi_out);
+    } break;
+    case SYMS_FORMAT_DWARF: {
+      DwRangeIter *iter_impl = (DwRangeIter *)iter->impl;
+      if (dw_range_iter_next(iter_impl, lo_out, hi_out)) {
+        is_next_valid = syms_true;
+      }
+    } break;
+    }
+  } break;
+  }
+
+  if (is_next_valid) {
+    SymsAddr rebase = syms_get_rebase(iter->instance);
+    *lo_out = rebase + *lo_out;
+    *hi_out = rebase + *hi_out;
+  }
+
+  return is_next_valid;
+}
+
+// -------------------------------------------------------------------------------- 
+
+SYMS_API syms_bool
+syms_proc_from_name_2(struct SymsInstance *instance, const char *name, syms_uint name_size, SymsProc *proc_out, SymsModID *mod_out)
+{
+  SymsDebugInfo *debug_info = &instance->debug_info;
+  SymsString left = syms_string_init(name, name_size);
+  SymsModID i;
+  for (i = 0; i < instance->mods_num; ++i) {
+    switch (debug_info->type) {
+    case SYMS_FORMAT_NULL: break;
+    case SYMS_FORMAT_PDB: {
+      if (syms_proc_from_name_pdb(syms_get_debug_info_pdb(debug_info), left, proc_out)) {
+        *mod_out = i;
+        proc_out->va = syms_rebase_addr(instance, proc_out->va);
+        proc_out->dbg_start_va = syms_rebase_addr(instance, proc_out->dbg_start_va);
+        proc_out->dbg_end_va = syms_rebase_addr(instance, proc_out->dbg_end_va);
+        return syms_true;
+      }
+    } break;
+    case SYMS_FORMAT_DWARF: {
+      SymsModInfo *mod = instance->mods + i;
+      syms_uint k;
+      for (k = 0; k < mod->procs->push_count; ++k) {
+        SymsProc *proc = syms_mod_info_find_proc(mod, k);
+        if (syms_string_cmp_strref(instance, left, proc->name_ref)) {
+          *mod_out = i;
+          *proc_out = *proc;
+          return syms_true;
+        }
+      }
+    } break;
+    }
+  }
+  return syms_false;
+}
+
+SYMS_API syms_bool
+syms_proc_from_name(struct SymsInstance *instance, const char *name, syms_uint name_size, SymsProc *proc_out)
+{
+  SymsModID dummy;
+  return syms_proc_from_name_2(instance, name, name_size, proc_out, &dummy);
+}
+
+SYMS_API syms_bool
+syms_img_sec_from_name(struct SymsInstance *instance, SymsString name, SymsSection *sec_out)
+{
+  SymsSecIter iter;
+  if (syms_sec_iter_init(instance, &iter)) {
+    while (syms_sec_iter_next(&iter, sec_out)) {
+      if (syms_string_cmp(sec_out->name, name)) {
+        return syms_true;
+      }
+    }
+  }
+  return syms_false;
+}
+
+SYMS_API syms_bool
+syms_va_to_secoff(SymsInstance *instance, SymsAddr va, syms_uint *isec_out, syms_uint *off_out)
+{
+  SymsSecIter iter;
+  if (syms_sec_iter_init(instance, &iter)) {
+    SymsSection sec;
+    syms_uint i = 0;
+    while (syms_sec_iter_next(&iter, &sec)) {
+      if (sec.va <= va && va < sec.va + sec.data_size) {
+        *isec_out = i;
+        *off_out = (syms_uint)sec.off;
+        return syms_true;
+      }
+      ++i;
+    }
+  }
+  return syms_false;
+}
+
+SYMS_API SymsErrorCode
+syms_trampoline_from_ip(struct SymsInstance *instance, SymsAddr ip, SymsAddr *ip_out)
+{
+  SymsDebugInfo *debug_info = &instance->debug_info;
+  SymsAddr rebase = syms_get_rebase(instance);
+  SymsAddr dest_rva = SYMS_ADDR_MAX;
+
+  syms_bool result = syms_false;
+  SymsErrorCode err = SYMS_ERR_INVALID_CODE_PATH;
+
+  if (ip >= rebase) {
+    switch (debug_info->type) {
+    case SYMS_FORMAT_NULL: break;
+    case SYMS_FORMAT_PDB: {
+      SymsAddr rva = ip - rebase;
+      result = syms_trampoline_from_ip_pdb(syms_get_debug_info_pdb(debug_info), rva, &dest_rva); 
+    } break;
+    case SYMS_FORMAT_DWARF: {
+      result = syms_trampoline_from_ip_dwarf(syms_get_debug_info_dwarf(debug_info), ip, &dest_rva);
+    } break;
+    }
+
+    if (ip_out) {
+      *ip_out = rebase + dest_rva;
+    }
+  }
+
+  if (result) {
+    err = SYMS_ERR_OK;
+  }
+  return err;
+}
+
+// -------------------------------------------------------------------------------- 
+
+SYMS_API SymsTypeID
+syms_type_id_for_kind(SymsTypeKind kind)
+{
+  SymsTypeID id;
+  id.kind = SYMS_TYPE_ID_BUILTIN;
+  *(SymsTypeKind *)id.impl = kind;
+  return id;
+}
+SYMS_INTERNAL SymsTypeID
+syms_type_id_null(void)
+{
+  SymsTypeID type_id;
+  syms_memzero(&type_id, sizeof(type_id));
+  return type_id;
+}
+
+SYMS_API SymsType
+syms_infer_base_type(SymsInstance *instance, const SymsType *type)
+{
+  SymsType base_type = *type;
+  SymsType next_type;
+  while (syms_infer_type(instance, base_type.next_id, &next_type)) {
+    if (next_type.kind == SYMS_TYPE_NULL) {
+      break;
+    }
+    base_type = next_type;
+    if (base_type.kind == SYMS_TYPE_PROC) {
+      break;
+    }
+  }
+  return base_type;
+}
+
+SYMS_API void
+syms_infer_builtin_type(SymsArch arch, SymsTypeKind type_kind, SymsType *type_out)
+{
+  const char *label = syms_get_type_kind_info(type_kind);
+  syms_uint bit_count = syms_get_type_kind_bitcount(arch, type_kind);
+
+  type_out->id        = syms_type_id_for_kind(type_kind);
+  type_out->next_id   = syms_type_id_null();
+  type_out->modifier  = 0;
+  type_out->kind      = type_kind;
+  type_out->is_fwdref = syms_false;
+  type_out->size      = (syms_uint)((float)bit_count / 8.0f + 0.5f);
+  type_out->name_ref  = syms_string_ref_str(syms_string_init_lit(label));
+
+  if (type_kind == SYMS_TYPE_CHAR) {
+    type_out->modifier |= SYMS_TYPE_MDFR_CHAR;
+  }
+}
+
+
+SYMS_API syms_int
+syms_typeid_cmp(const SymsTypeID *a, const SymsTypeID *b)
+{
+  syms_int cmp_result = 0xdeadc0de;
+  if (a->kind == b->kind) {
+    switch (a->kind) {
+    case SYMS_TYPE_ID_NULL: break;
+    case SYMS_TYPE_ID_ELF:  break;
+    case SYMS_TYPE_ID_PDB: {
+      pdb_cv_itype impl_a = syms_typeid_to_pdb(a); 
+      pdb_cv_itype impl_b = syms_typeid_to_pdb(b);
+      cmp_result = syms_typeid_cmp_pdb(impl_a, impl_b); 
+    } break;
+    case SYMS_TYPE_ID_DW: {
+      const DwTag *impl_a = syms_typeid_to_dw(a);
+      const DwTag *impl_b = syms_typeid_to_dw(b);
+      cmp_result = syms_typeid_cmp_dw(*impl_a, *impl_b);
+    } break;
+    case SYMS_TYPE_ID_BUILTIN: {
+      SymsTypeKind impl_a = *(SymsTypeKind *)a->impl;
+      SymsTypeKind impl_b = *(SymsTypeKind *)b->impl;
+      if (impl_a < impl_b) {
+        cmp_result = -1;
+      } else if (impl_a > impl_b) {
+        cmp_result = +1;
+      } else {
+        cmp_result = 0;
+      }
+    } break;
+    case SYMS_TYPE_ID_COUNTER: {
+      syms_uint impl_a = *(syms_uint *)a->impl;
+      syms_uint impl_b = *(syms_uint *)b->impl;
+      if (impl_a < impl_b) {
+        cmp_result = -1;
+      } else if (impl_a > impl_b) {
+        cmp_result = +1;
+      } else {
+        cmp_result = 0;
+      }
+    } break;
+    }
+  } else {
+    SYMS_ASSERT_FAILURE("unable to compare types with conflicting kinds");
+  }
+  return cmp_result;
+}
+
+SYMS_API SymsType
+syms_make_ptr_type(SymsArch arch, const SymsTypeID *ref)
+{
+  SymsType type;
+  type.id = syms_type_id_null();
+  type.next_id = *ref;
+  type.kind = SYMS_TYPE_PTR;
+  type.size = syms_get_addr_size_ex(arch);
+  type.is_fwdref = syms_false;
+  type.modifier = 0;
+  type.name_ref = syms_string_ref_null();
+  type.u.ptr.mode = SYMS_PTR_MODE_NORMAL;
+  return type;
+}
+
+SYMS_API syms_bool
+syms_infer_type(struct SymsInstance *instance, SymsTypeID type_id, SymsType *type_out)
+{
+  syms_bool is_ok = syms_false;
+  SymsDebugInfo *debug_info = &instance->debug_info;
+
+  switch (type_id.kind) {
+  case SYMS_TYPE_ID_NULL: break;
+  case SYMS_TYPE_ID_ELF: break;
+  case SYMS_TYPE_ID_PDB: {
+    SymsDebugInfoPdb *impl_info = syms_get_debug_info_pdb(debug_info);
+    is_ok = syms_infer_type_pdb(impl_info, type_id, type_out);
+    if (is_ok) {
+      if (type_out->kind == SYMS_TYPE_VOID || type_out->kind == SYMS_TYPE_PROC) {
+        type_out->size = syms_get_addr_size(instance);
+      }
+    }
+  } break;
+  case SYMS_TYPE_ID_DW: {
+    SymsDebugInfoDwarf *impl_info = syms_get_debug_info_dwarf(debug_info);
+    is_ok = syms_infer_type_dwarf(impl_info, type_id, type_out);
+  } break;
+  case SYMS_TYPE_ID_BUILTIN: {
+    SymsTypeKind kind = *(SymsTypeKind *)type_id.impl;
+    SymsArch arch = syms_get_arch(instance);
+    syms_infer_builtin_type(arch, kind, type_out);
+    is_ok = syms_true;
+  } break;
+  }
+
+  return is_ok;
+}
+
+SYMS_API syms_uint
+syms_get_type_size(struct SymsInstance *instance, SymsType *type)
+{
+  syms_uint result = 0;
+  if (type->kind == SYMS_TYPE_ARR) {
+    SymsTypeID next_id = type->next_id;
+    SymsType next_type;
+    syms_uint element_size = 0;
+    while (syms_infer_type(instance, next_id, &next_type)) {
+      if (next_type.kind != SYMS_TYPE_ARR) {
+        element_size = next_type.size;
+        break;
+      }
+      next_id = next_type.next_id;
+    }
+    result = (syms_uint)element_size;
+    next_type = *type;
+    do {
+      if (next_type.kind != SYMS_TYPE_ARR) {
+        break;
+      }
+      result *= next_type.size;
+      next_id = next_type.next_id;
+    } while (syms_infer_type(instance, next_id, &next_type));
+    
+  } else {
+    result = type->size;
+  }
+  return result;
+}
+
+SYMS_API SymsErrorCode
+syms_decode_location(struct SymsInstance *instance, SymsEncodedLocation *loc, 
+    void *regread_context, syms_regread_sig *regread_callback,
+    void *memread_context, syms_memread_sig *memread_callback, 
+    SymsLocation *loc_out)
+{
+  SymsDebugInfo *debug_info = &instance->debug_info;
+  SymsErrorCode decode_result = SYMS_ERR_INVALID_CODE_PATH;
+
+  SymsMemread memread;
+  SymsRegread regread;
+  
+  SymsAddr orig_rebase;
+  SymsAddr rebase;
+
+  memread.result = SYMS_ERR_INVALID_CODE_PATH;
+  memread.context = memread_context;
+  memread.callback = memread_callback;
+
+  regread.context = regread_context;
+  regread.callback = regread_callback;
+  regread.arch = syms_get_arch(instance);
+
+  orig_rebase = syms_get_orig_rebase(instance);
+  rebase      = syms_get_rebase(instance);
+
+  switch (debug_info->type) {
+  case SYMS_FORMAT_NULL : break;
+  case SYMS_FORMAT_PDB  : { 
+    pdb_encoded_location *loc_impl = (pdb_encoded_location *)loc->impl;
+    decode_result = syms_decode_location_pdb(orig_rebase, rebase, loc_impl, &regread, &memread, loc_out);
+  } break;
+  case SYMS_FORMAT_DWARF: {
+    DwEncodedLocation *loc_impl = (DwEncodedLocation *)loc->impl;
+    decode_result = syms_decode_location_dwarf(orig_rebase, rebase, loc_impl, 0, &memread, loc_out); 
+  } break;
+  default: SYMS_INVALID_CODE_PATH;
+  }
+
+  if (SYMS_RESULT_FAIL(decode_result)) {
+    // NOTE(nick): Not a great solution to detect whether or not decoding failed because of
+    // a callback. 
+    //
+    // Debugger arch allows a failure at any point of execution, because
+    // target can be on another machine and downloading entire RAM is not a reasonable solution,
+    // so at any point debugger expects to receive a "maybe" error code that says
+    // that execution didn't fail nor succeeded and it should try again later.
+    if (memread.result == SYMS_ERR_MAYBE) {
+      decode_result = memread.result;
+    }
+  }
+  
+  return decode_result;
+}
+
+SYMS_INTERNAL syms_uint
+syms_decode_location_bytes(SymsInstance *instance, 
+               SymsEncodedLocation *encoded_loc, 
+               void *regread_ctx, syms_regread_sig *regread,
+               void *memread_ctx, syms_memread_sig *memread,
+               void *buffer, syms_uint buffer_size)
+{
+  SymsLocation loc;
+  syms_uint decode_size = 0;
+  SymsErrorCode decode_result = syms_decode_location(instance,
+                     encoded_loc,
+                     regread_ctx, regread,
+                     memread_ctx, memread, 
+                     &loc);
+  if (SYMS_RESULT_OK(decode_result)) {
+    void const *src_data = 0;
+    syms_uint src_len = 0;
+    switch (loc.kind) {
+    case SYMS_LOCATION_IMPLICIT: {
+      src_data = (void const *)&loc.u.implicit.data;
+      src_len = loc.u.implicit.len;
+    } break;
+    case SYMS_LOCATION_INDIRECT: {
+      src_data = (void const *)loc.u.indirect.data;
+      src_len = loc.u.indirect.len;
+    } break;
+    case SYMS_LOCATION_VA: {
+      src_data = (void const *)&loc.u.va;
+      src_len = sizeof(loc.u.va);
+    } break;
+    default: SYMS_INVALID_CODE_PATH; break;
+    }
+    if (src_data) {
+      if (src_len <= buffer_size) {
+        syms_memcpy(buffer, src_data, src_len);
+        decode_size = src_len;
+      }
+    }
+  }
+  return decode_size;
+}
+
+SYMS_API syms_uint
+syms_decode_location32(SymsInstance *instance, 
+             SymsEncodedLocation *encoded_loc,
+             void *regread_context, syms_regread_sig *regread_callback,
+             void *memread_context, syms_memread_sig *memread_callback)
+{
+  syms_uint value = 0;
+  syms_uint decode_size = syms_decode_location_bytes(instance,
+                  encoded_loc,
+                  regread_context, regread_callback,
+                  memread_context, memread_callback, 
+                  &value, sizeof(value));
+  SYMS_ASSERT_ALWAYS(decode_size > 0);
+  return value;
+}
+
+SYMS_API U64
+syms_decode_location64(SymsInstance *instance, 
+             SymsEncodedLocation *encoded_loc,
+             void *regread_context, syms_regread_sig *regread_callback,
+             void *memread_context, syms_memread_sig *memread_callback)
+{
+  U64 value = 0;
+  U64 decode_size = syms_decode_location_bytes(instance, 
+                    encoded_loc, 
+                    regread_context, regread_callback,
+                    memread_context, memread_callback, 
+                    &value, sizeof(value));
+  SYMS_ASSERT_ALWAYS(decode_size > 0);
+  return value;
+}
+
+SYMS_INTERNAL U8
+syms_chksum_get_num_bytes(SymsChecksumType type)
+{
+  U8 result;
+  switch (type) {
+  case SYMS_CHECKSUM_MD5:    result = 16; break;
+  case SYMS_CHECKSUM_SHA1:   result = 20; break;
+  case SYMS_CHECKSUM_SHA256: result = 32; break;
+  default:                   result = 0;  break;
+  }
+  return result;
+}
+
+SYMS_API syms_bool
+syms_src_to_va(struct SymsInstance *instance, char *filename, syms_uint filename_len, syms_uint ln, SymsSourceFileMap *map_out)
+{
+  SymsString file = syms_string_init(filename, filename_len);
+  syms_bool is_mapped = syms_false;
+  SymsModID i;
+  for (i = 0; i < instance->mods_num; ++i) {
+    SymsModInfo *mod = instance->mods + i;
+    is_mapped = syms_line_table_map_src(&mod->line_table, instance, file, ln, map_out);
+    if (is_mapped) {
+      break;
+    }
+  }
+  return is_mapped;
+}
+
+SYMS_API syms_bool
+syms_va_to_src(struct SymsInstance *instance, SymsAddr va, SymsSourceFileMap *map_out)
+{
+  syms_bool is_result_valid = syms_false;
+  SymsAddr nearest_delta = SYMS_ADDR_MAX;
+  SymsModID i;
+  for (i = 0; i < instance->mods_num; ++i) {
+    SymsModInfo *mod = &instance->mods[i];
+    SymsSourceFileMap map;
+    if (syms_line_table_map_va(&mod->line_table, va, &map)) {
+      SymsAddr delta;
+
+      SYMS_ASSERT(va >= map.line.va);
+      delta = va - map.line.va;
+      if (delta < nearest_delta) {
+        *map_out = map;
+        nearest_delta = delta;
+        is_result_valid = syms_true;
+        if (delta == 0)
+          break;
+      }
+    }
+  }
+  return is_result_valid;
+}
+
+SYMS_API syms_bool
+syms_proc_from_va_2(struct SymsInstance *instance, SymsAddr va, SymsProc *proc_out, SymsModID *mod_out)
+{
+  syms_bool is_mapped = syms_false;
+  SymsModID mod_id; 
+  for (mod_id = 0; mod_id < instance->mods_num; ++mod_id) {
+    SymsModInfo *mod = &instance->mods[mod_id];
+    SymsRangeMap *map = syms_rangemap_search(mod->rangemap, va);
+    if (map) {
+      SymsProc *proc = syms_mod_info_find_proc(mod, map->id);
+      SYMS_ASSERT(map->id < mod->procs->push_count);
+      *proc_out = *proc;
+      *mod_out = mod_id;
+      is_mapped = syms_true;
+      break;
+    }
+  }
+  return is_mapped;
+}
+
+SYMS_API syms_bool
+syms_proc_from_va(SymsInstance *instance, SymsAddr va, SymsProc *proc_out)
+{
+  SymsModID dummy;
+  return syms_proc_from_va_2(instance, va, proc_out, &dummy);
+}
+
+SYMS_API syms_uint
+syms_get_proc_count(struct SymsInstance *instance)
+{
+  syms_uint count = 0;
+  SymsModID i;
+  for (i = 0; i < instance->mods_num; ++i) {
+    count += instance->mods[i].procs->push_count;
+  }
+  return count;
+}
+
+SYMS_API syms_uint
+syms_get_line_count(struct SymsInstance *instance)
+{
+  syms_uint line_count = 0;
+  SymsModID i;
+  for (i = 0; i < instance->mods_num; ++i) {
+    SymsModInfo *mod = &instance->mods[i];
+    line_count += mod->line_table.line_count;
+  }
+  return line_count;
+}
+
+SYMS_API syms_uint
+syms_get_addr_size(SymsInstance *instance)
+{
+  SymsArch arch = syms_get_arch(instance);
+  return syms_get_addr_size_ex(arch);
+}
+
+SYMS_API SymsErrorCode
+syms_build_module(SymsInstance *instance, SymsModID mod_id, SymsArena *arena)
+{
+  SymsArenaFrame frame = syms_arena_frame_begin(arena);
+  SymsModInfo *mod = instance->mods + mod_id;
+  SymsDebugInfo *debug_info = &instance->debug_info;
+  SymsErrorCode err = SYMS_ERR_OK;
+
+  if (mod_id >= instance->mods_num) {
+    err = SYMS_ERR_INVAL;
+  }
+
+  if (SYMS_RESULT_OK(err)) {
+    // Build protion of .debug_abbrev that corresponds to input module index.
+    if (debug_info->type == SYMS_FORMAT_DWARF) {
+      DwCompileUnit *cu = (DwCompileUnit *)mod->header.impl;
+      if (!dw_build_abbrev(cu, arena)) {
+        err = SYMS_ERR_INVALID_CODE_PATH;
+      }
+    }
+  }
+
+  // Procedure symbols
+  if (SYMS_RESULT_OK(err)) {
+    mod->procs = syms_block_allocator_init(sizeof(SymsProc), sizeof(SymsProc)*1024, arena);
+    SymsProcIter proc_iter;
+    if (syms_proc_iter_init(instance, &proc_iter, &mod->header)) {
+      SymsProc temp;
+      while (syms_proc_iter_next(&proc_iter, &temp)) {
+        SymsProc *proc = (SymsProc *)syms_block_allocator_push(mod->procs);
+        if (!proc) {
+          err = SYMS_ERR_NOMEM;
+          break;
+        }
+        *proc = temp;
+      }
+    }
+  } else {
+    mod->procs = syms_block_allocator_init(0, 0, arena);
+  }
+  syms_block_allocator_build_index_table(mod->procs);
+
+  // Load range map for procedures
+  if (SYMS_RESULT_OK(err)) {
+    mod->rangemap = syms_block_allocator_init(sizeof(SymsRangeMap), sizeof(SymsRangeMap)*1024, arena);
+    syms_uint proc_id;
+    for (proc_id = 0; proc_id < mod->procs->push_count; ++proc_id) {
+      SymsProc *proc = syms_mod_info_find_proc(mod, proc_id);
+      SymsRangeIter range_iter;
+      if (syms_range_iter_init(&range_iter, instance, &proc->range)) {
+        for (;;) {
+          SymsAddr lo, hi;
+          SymsRangeMap *rangemap;
+
+          if (!syms_range_iter_next(&range_iter, &lo, &hi)) {
+            break;
+          }
+          rangemap = (SymsRangeMap*)syms_block_allocator_push(mod->rangemap);
+          if (!rangemap) {
+            break;
+          }
+          rangemap->id = proc_id;
+          rangemap->lo = lo;
+          rangemap->hi = hi;
+        }
+      } else {
+        SymsRangeMap *rangemap = (SymsRangeMap *)syms_block_allocator_push(mod->rangemap);
+        if (!rangemap) {
+          break;
+        }
+        rangemap->id = proc_id;
+        rangemap->lo = proc->va;
+        rangemap->hi = proc->va + proc->len;
+      }
+    }
+  } else {
+    mod->rangemap = syms_block_allocator_init(0, 0, arena);
+  }
+  syms_block_allocator_build_index_table(mod->rangemap);
+  syms_block_allocator_sort(mod->rangemap, syms_qsort_rangemap);
+
+  if (SYMS_RESULT_OK(err)) {
+    err = syms_line_table_build(&mod->line_table, instance, &mod->header, arena);
+  }
+
+  if (arena->flags & SYMS_ARENA_FLAG_ALLOC_FAILED) {
+    err = SYMS_ERR_NOMEM;
+  }
+
+  if (SYMS_RESULT_FAIL(err)) {
+    mod->procs = 0;
+    mod->rangemap = 0;
+    syms_memzero(&mod->line_table, sizeof mod->line_table);
+    syms_arena_frame_end(frame);
+  }
+
+  return err;
+}
+
+SYMS_API syms_uint
+syms_get_module_build_count(struct SymsInstance *instance)
+{
+  return instance->mods_num;
+}
+
+SYMS_API SymsModID
+syms_infer_global_data_module(struct SymsInstance *instance, SymsGlobal *gdata)
+{
+  SymsModID mod_id = SYMS_INVALID_MOD_ID;
+  SymsDebugInfo *debug_info = &instance->debug_info;
+  switch (debug_info->type) {
+  case SYMS_FORMAT_NULL: break;
+  case SYMS_FORMAT_PDB: {
+    mod_id = syms_infer_global_data_module_pdb(instance, syms_get_debug_info_pdb(debug_info), gdata);
+  } break;
+  case SYMS_FORMAT_DWARF: {
+    // TODO(nick): DWARF
+  } break;
+  case SYMS_FORMAT_ELFSYMTAB: {
+    // TODO(nick): ELFSYM
+  } break;
+  }
+  return mod_id;
+}
+
+SYMS_API SymsModInfo *
+syms_get_mod(struct SymsInstance *instance, SymsModID mod_id)
+{
+  SymsModInfo *mod = 0;
+  if (mod_id < instance->mods_num) {
+    mod = instance->mods + mod_id;
+  }
+  return mod;
+}
+
+SYMS_API const char *
+syms_get_version(void)
+{
+  return SYMS_VERSION_STR;
+}
+
+// ----------------------------------------------------------------------------
+
+SYMS_INTERNAL SymsSymbol *
+syms_arena_push_symbol(SymsSymbol **first, SymsSymbolKind kind, SymsArena *arena)
+{
+  SymsSymbol *sym = syms_arena_push_struct(arena, SymsSymbol);
+  if (sym) {
+    sym->kind = kind;
+    sym->next = *first;
+    *first = sym;
+  }
+  return sym;
+}
+
+SYMS_INTERNAL SymsSymbolProc *
+syms_arena_push_symbol_proc(SymsSymbol **first, SymsArena *arena, SymsAddr addr, syms_uint len)
+{
+  SymsSymbol *sym = syms_arena_push_symbol(first, SYMS_SYMBOL_PROC, arena);
+  SymsSymbolProc *proc = 0;
+  if (sym) {
+    proc = &sym->u.proc;
+    proc->range_lo  = addr;
+    proc->range_hi  = addr + len;
+    proc->dbg_start = proc->range_lo;
+    proc->dbg_end   = proc->range_hi;
+  }
+  return proc;
+}
+
+SYMS_INTERNAL SymsSymbolConstData *
+syms_arena_push_symbol_const_data(SymsSymbol **first, SymsArena *arena)
+{
+  SymsSymbol *sym = syms_arena_push_symbol(first, SYMS_SYMBOL_CONST_DATA, arena);
+  if (sym)
+    return &sym->u.cdata;
+  return 0;
+}
+
+SYMS_INTERNAL SymsSymbolStaticData *
+syms_arena_push_symbol_static_data(SymsSymbol **first, SymsArena *arena)
+{
+  SymsSymbol *sym = syms_arena_push_symbol(first, SYMS_SYMBOL_STATIC_DATA, arena);
+  if (sym)
+    return &sym->u.sdata;
+  return 0;
+}
+
+SYMS_INTERNAL SymsSymbolTypeInfo *
+syms_arena_push_symbol_type(SymsSymbol **first, SymsArena *arena)
+{
+  SymsSymbol *sym = syms_arena_push_symbol(first, SYMS_SYMBOL_TYPE_INFO, arena);
+  if (sym)
+    return &sym->u.type_info;
+  return 0;
+}
+
+SYMS_INTERNAL void
+syms_export_srcmap(struct SymsInstance *instance, SymsSourceFileMap *map, SymsSymbolSrcmap *map_out, SymsArena *arena)
+{
+  syms_uint chksum_size = 0;
+  u8 *chksum = 0;
+  // copy path
+  map_out->path = syms_read_strref_arena(instance, &map->file.name, arena);
+  // copy file check sum
+  chksum_size = syms_chksum_get_num_bytes(map->file.chksum_type);
+  if (chksum_size) {
+    chksum = syms_arena_push_array(arena, u8, chksum_size);
+    syms_memcpy(chksum, map->file.chksum, chksum_size);
+  }
+  //
+  map_out->chksum       = (void *)chksum;
+  map_out->chksum_size  = chksum_size;
+  map_out->chksum_type  = map->file.chksum_type;
+  map_out->addr         = map->line.va;
+  map_out->ln           = map->line.ln;
+  map_out->col          = map->line.col;
+  map_out->instructions_size = map->instructions_size;
+}
+
+SYMS_INTERNAL void
+syms_export_symbol_section(SymsInstance *instance, SymsSection *sec, SymsSymbolSection *sec_out, SymsArena *arena)
+{ (void)instance;
+  sec_out->name = syms_string_to_cstr(sec->name, arena);
+  sec_out->off = sec->off;
+  sec_out->va = sec->va;
+  sec_out->length = sec->data_size;
+}
+
+SYMS_API SymsSymbolSection *
+syms_img_sec_from_name_ex(struct SymsInstance *instance, const char *name, SymsArena *arena)
+{
+  SymsString name_str = syms_string_init_lit(name);
+  SymsSymbolSection *sec_out = NULL;
+  SymsSection sec;
+
+  if (syms_img_sec_from_name(instance, name_str, &sec)) {
+    sec_out = syms_arena_push_struct(arena, SymsSymbolSection);
+    syms_export_symbol_section(instance, &sec, sec_out, arena);
+  }
+
+  return sec_out;
+}
+
+SYMS_INTERNAL SymsEncodedLocation *
+syms_export_encoded_location(SymsEncodedLocation loc, SymsArena *arena)
+{
+  SymsEncodedLocation *r = syms_arena_push_struct(arena, SymsEncodedLocation);
+  *r = loc;
+  return (SymsEncodedLocation *)r;
+}
+
+SYMS_INTERNAL SymsTypeID *
+syms_export_typeid(SymsTypeID id, SymsArena *arena)
+{
+  SymsTypeID *r = syms_arena_push_struct(arena, SymsTypeID);
+  *r = id;
+  return r;
+}
+
+SYMS_INTERNAL void
+syms_export_proc(SymsInstance *instance, SymsProc *proc, SymsSourceFileMap *srcmap, SymsSymbolProc *proc_out, SymsArena *arena)
+{
+  proc_out->range_lo  = proc->va;
+  proc_out->range_hi  = proc->va + proc->len;
+  proc_out->dbg_start = proc->dbg_start_va;
+  proc_out->dbg_end   = proc->dbg_end_va;
+  proc_out->label     = syms_read_strref_arena(instance, &proc->name_ref, arena);
+  proc_out->type_id   = syms_export_typeid(proc->type_id, arena);
+  if (srcmap)
+    syms_export_srcmap(instance, srcmap, &proc_out->srcmap, arena);
+  else
+    syms_memzero(&proc_out->srcmap, sizeof(proc_out->srcmap));
+}
+
+SYMS_INTERNAL void
+syms_export_inline_site(SymsInstance *instance, SymsInlineSite *site, SymsSymbolProc *proc_out, SymsArena *arena)
+{
+  syms_export_proc(instance, &site->inlinee, &site->src, proc_out, arena);
+}
+
+SYMS_INTERNAL SymsSymbolVar *
+syms_export_locals(SymsInstance *instance, SymsProc *proc, SymsAddr addr, syms_uint *local_count_out, SymsArena *arena)
+{
+  char stack[1024];
+  SymsSymbolVar *first_var = 0;
+  SymsLocalIter local_iter;
+  if (syms_local_iter_init(instance, &local_iter, proc, stack, sizeof(stack))) {
+    SymsSymbolVar **next_var = &first_var;
+    SymsLocalExport exported_data;
+
+    while (syms_local_iter_next(&local_iter, &exported_data)) {
+      if (exported_data.type == SYMS_LOCAL_EXPORT_VAR) {
+        SymsRangeIter range_iter;
+        SymsSymbolVar *var = syms_arena_push_struct(arena, SymsSymbolVar);
+        syms_bool var_inrange = syms_false;
+        if (!var) break;
+
+        if (syms_range_iter_init(&range_iter, instance, &exported_data.u.var.range)) {
+          SymsAddr lo, hi;
+          while (syms_range_iter_next(&range_iter, &lo, &hi)) {
+            if (lo <= addr && addr < hi) {
+              var_inrange = syms_true;
+              break;
+            }
+          }
+        }
+
+        if (var_inrange) {
+          *next_var = var;
+          next_var = &var->next;
+
+          var->label      = syms_read_strref_arena(instance, &exported_data.u.var.name_ref, arena);
+          var->flags      = exported_data.u.var.flags;
+          var->encoded_va = syms_export_encoded_location(exported_data.u.var.encoded_va, arena);
+          var->type_id    = syms_export_typeid(exported_data.u.var.type_id, arena);
+          var->next       = 0;
+
+          *local_count_out += 1;
+        }
+      }
+    }
+  }
+  return first_var;
+}
+
+SYMS_INTERNAL SymsTypeList
+syms_export_arglist(SymsInstance *instance, SymsTypeID arglist_id, SymsArena *arena)
+{
+  SymsTypeList r;
+  SymsTypeIDNode **next;
+  SymsType arglist_type;
+
+  r.first = 0;
+  r.count = 0;
+  //
+  next = &r.first;
+
+  if (syms_infer_type(instance, arglist_id, &arglist_type)) {
+    SymsArgIter arg_iter;
+    if (syms_arg_iter_init(instance, &arg_iter, &arglist_type)) {
+      SymsTypeID arg_id;
+      while (syms_arg_iter_next(&arg_iter, &arg_id)) {
+        SymsTypeIDNode *node = syms_arena_push_struct(arena, SymsTypeIDNode);
+        node->data = syms_export_typeid(arg_id, arena);
+        node->next = 0;
+        //
+        *next = node;
+        next = &node->next;
+        //
+        r.count += 1;
+      }
+    }
+  }
+
+  return r;
+}
+
+SYMS_INTERNAL void
+syms_export_symbol_const_data(struct SymsInstance *instance, SymsConst *cdata, SymsSymbolConstData *cdata_out, SymsArena *arena)
+{
+  syms_uint buflen;
+  u8 *buf;
+
+  cdata_out->label = syms_read_strref_arena(instance, &cdata->name, arena);
+  cdata_out->type_id = syms_export_typeid(cdata->type_id, arena);
+  
+  buflen = cdata->value_len;
+  buf = syms_arena_push_array(arena, u8, cdata->value_len + 1);
+  syms_memcpy(buf, cdata->value, buflen);
+  buf[buflen] = 0;
+
+  cdata_out->value_size = buflen;
+  cdata_out->value = (void *)buf;
+}
+
+SYMS_INTERNAL void
+syms_export_global_data(struct SymsInstance *instance, SymsGlobal *gdata, SymsSymbolStaticData *gdata_out, SymsArena *arena)
+{
+  SymsLocation loc;
+
+  gdata_out->label = syms_read_strref_arena(instance, &gdata->name, arena);
+  gdata_out->type_id = syms_export_typeid(gdata->type_id, arena);
+  gdata_out->addr = 0;
+
+  if (syms_decode_location(instance, &gdata->encoded_va, 0, 0, 0, 0, &loc)) {
+    if (loc.kind == SYMS_LOCATION_VA) {
+      gdata_out->addr = loc.u.va;
+    }
+  }
+}
+
+SYMS_API SymsErrorCode
+syms_export_type(struct SymsInstance *instance, SymsTypeID type_id, SymsSymbolTypeInfo *ti, SymsArena *arena)
+{
+  SymsErrorCode err = SYMS_ERR_OK;
+  SymsType type;
+
+  if (syms_infer_type(instance, type_id, &type)) {
+    //ti = syms_arena_push_struct(arena, SymsSymbolTypeInfo);
+    ti->modifier = type.modifier;
+    ti->kind = type.kind;
+    ti->label = syms_read_strref_arena(instance, &type.name_ref, arena);
+    ti->type_id = syms_export_typeid(type_id, arena);
+    ti->size = (int)syms_get_type_size(instance, &type);
+    switch (type.kind) {
+    case SYMS_TYPE_PROC: {
+      ti->u.proc.arg_count = type.u.proc.arg_count;
+      ti->u.proc.arglist_id = syms_export_typeid(type.u.proc.arglist_type_id, arena);
+      ti->u.proc.return_id = syms_export_typeid(type.u.proc.ret_type_id, arena);
+    } break;
+    case SYMS_TYPE_METHOD: {
+      ti->u.method.arg_count = type.u.method.arg_count;
+      ti->u.method.class_id = syms_export_typeid(type.u.method.class_type_id, arena);
+      ti->u.method.return_id = syms_export_typeid(type.u.method.ret_type_id, arena);
+      ti->u.method.this_id = syms_export_typeid(type.u.method.this_type_id, arena);
+      ti->u.method.arglist_id = syms_export_typeid(type.u.method.arglist_type_id, arena);
+    } break;
+    case SYMS_TYPE_ARR: {
+      ti->u.array.entry_count = type.size;
+      ti->u.array.entry_id = syms_export_typeid(type.next_id, arena);
+    } break;
+    case SYMS_TYPE_BITFIELD: {
+      ti->u.bitfield.base_type_id = syms_export_typeid(type.u.bitfield.base_type_id, arena);
+      ti->u.bitfield.len = type.u.bitfield.len;
+      ti->u.bitfield.pos = type.u.bitfield.pos;
+    } break;
+    case SYMS_TYPE_UNION:
+    case SYMS_TYPE_STRUCT: 
+    case SYMS_TYPE_ENUM: 
+	case SYMS_TYPE_CLASS: {
+      SymsMemberIter iter;
+      SymsMember member;
+
+      if (syms_member_iter_init(instance, &iter, &type)) {
+        syms_uint count = 0;
+        syms_uint member_index;
+
+        while (syms_member_iter_next(&iter, &member)) {
+          count += 1;
+        }
+        ti->u.udt.count = count;
+        ti->u.udt.base = syms_arena_push_array(arena, SymsTypeMember, count);
+        if (!ti->u.udt.base) {
+          err = SYMS_ERR_INVALID_CODE_PATH;
+          break;
+        }
+
+        member_index = 0;
+        if (!syms_member_iter_init(instance, &iter, &type)) {
+          SYMS_ASSERT_PARANOID("unable to reset member iterator");
+          err = SYMS_ERR_INVALID_CODE_PATH;
+          break;
+        }
+        for (;;) {
+          SymsTypeMember *member_out;
+
+          if (!syms_member_iter_next(&iter, &member)) {
+            break;
+          }
+
+          member_out          = ti->u.udt.base + member_index;
+          member_out->label   = syms_read_strref_arena(instance, &member.name_ref, arena);
+          member_out->type_id = syms_export_typeid(member.type_id, arena);
+          member_out->access  = member.access;
+          member_out->type    = member.type;
+          member_out->modifier = member.modifier;
+
+          switch (member.type) {
+          case SYMS_MEMBER_TYPE_STATIC_DATA: {
+            /* empty */
+          } break;
+          case SYMS_MEMBER_TYPE_NESTED_TYPE: {
+            /* empty */ 
+          } break;
+
+          case SYMS_MEMBER_TYPE_ENUM: {
+            SymsType value_type;
+            if (syms_infer_type(instance, member.type_id, &value_type)) {
+              SymsUMM enum_value;
+
+              switch (value_type.kind) {
+              case SYMS_TYPE_INT8:  enum_value = (SymsUMM)*(s8  *)member.u.enum_value; break;
+              case SYMS_TYPE_INT16: enum_value = (SymsUMM)*(s16 *)member.u.enum_value; break;
+              case SYMS_TYPE_INT32: enum_value = (SymsUMM)*(s32 *)member.u.enum_value; break;
+              case SYMS_TYPE_INT64: enum_value = (SymsUMM)*(s64 *)member.u.enum_value; break;
+
+              case SYMS_TYPE_UINT8:  enum_value = (SymsUMM)*(u8  *)member.u.enum_value; break;
+              case SYMS_TYPE_UINT16: enum_value = (SymsUMM)*(u16 *)member.u.enum_value; break;
+              case SYMS_TYPE_UINT32: enum_value = (SymsUMM)*(u32 *)member.u.enum_value; break;
+              case SYMS_TYPE_UINT64: enum_value = (SymsUMM)*(u64 *)member.u.enum_value; break;
+
+              default: {
+                SYMS_ASSERT_PARANOID("unable to read enum value -- unhandled case");
+                enum_value = SYMS_ADDR_MAX; 
+              } break;
+              }
+
+              member_out->u.enum_value.value = enum_value;
+            }
+          } break;
+
+          case SYMS_MEMBER_TYPE_METHOD: {
+            member_out->u.method.vbaseoff = member.u.method.vbaseoff;
+          } break;
+
+          case SYMS_MEMBER_TYPE_BASE_CLASS: {
+            member_out->u.base_class.offset = member.u.base_class.offset;
+          } break;
+
+          case SYMS_MEMBER_TYPE_DATA: {
+            member_out->u.data.offset = member.u.data_offset;
+          } break;
+
+          case SYMS_MEMBER_TYPE_VTABLE: {
+            // TODO: Move this code to syms_pdb_api.c
+            if (instance->debug_info.type == SYMS_FORMAT_PDB) {
+              SymsDebugInfoPdb *pdb_info = syms_get_debug_info_pdb(&instance->debug_info);
+              pdb_context *pdb = &pdb_info->context;
+              pdb_cv_itype vtabptr_itype = syms_typeid_to_pdb(&member.type_id);
+              pdb_type vtabptr;
+
+              if (pdb_infer_itype(pdb, vtabptr_itype, &vtabptr)) {
+                pdb_type vtshape;
+
+                if (pdb_infer_itype(pdb, vtabptr.next_cv_itype, &vtshape)) {
+                  SymsVTableEntryType *vtable = syms_arena_push_array(arena, SymsVTableEntryType, vtshape.u.vtshape.count);
+                  syms_uint i;
+                  pdb_uint read_offset = 0;
+                  u8 packed_desc = 0;
+
+                  for (i = 0; i < vtshape.u.vtshape.count; ++i) {
+                    syms_uint k;
+                    syms_uint desc;
+                    SymsVTableEntryType our_desc;
+
+                    k = i % 2;
+                    if (k == 0) {
+                      pdb_uint read_size = pdb_pointer_read(pdb, &vtshape.u.vtshape.ptr, read_offset, &packed_desc, sizeof(packed_desc));
+                      SYMS_ASSERT(read_size == sizeof(packed_desc));
+                    }
+
+                    /* each virtual table entry is packed as 4 bit integer */
+                    desc = (packed_desc >> k*4) & 0xf;
+
+                    /* conver CodevView vtable entry format to ours */
+                    switch (desc) {
+                    case PDB_CV_VTS_NEAR:   our_desc = SYMS_VTABLE_ENTRY_PTR16;     break;
+                    case PDB_CV_VTS_NEAR32: our_desc = SYMS_VTABLE_ENTRY_PTR32;     break;
+                    case PDB_CV_VTS_FAR:    our_desc = SYMS_VTABLE_ENTRY_SEGOFF16;  break;
+                    case PDB_CV_VTS_FAR32:  our_desc = SYMS_VTABLE_ENTRY_SEGOFF32;  break;
+                    case PDB_CV_VTS_THIN:   our_desc = SYMS_VTABLE_ENTRY_THIN;      break;
+                    case PDB_CV_VTS_OUTER:  our_desc = SYMS_VTABLE_ENTRY_OUTER;     break;
+                    case PDB_CV_VTS_META:   our_desc = SYMS_VTABLE_ENTRY_META;      break;
+                    default:                our_desc = SYMS_VTABLE_ENTRY_NULL;      break;
+                    }
+
+                    /* export entry */
+                    vtable[i] = our_desc;
+                  }
+
+                  member_out->u.vtab.offset = member.u.vtab.offset;
+                  member_out->u.vtab.count = vtshape.u.vtshape.count;
+                  member_out->u.vtab.base = vtable;
+                }
+              }
+            } else {
+              SYMS_ASSERT_FAILURE_PARANOID("Virtual table is not support for this debug info format");
+            }
+          } break;
+
+          default: {
+            SYMS_ASSERT_PARANOID("unable to export UDT member -- unhandled case");
+          } break;
+          }
+
+          member_index += 1;
+        }
+      }
+    } break;
+    }
+  } else {
+    err = SYMS_ERR_INVALID_CODE_PATH;
+  }
+
+  return err;
+}
+
+SYMS_API const char *
+syms_format_type(SymsInstance *instance, SymsType *type, SymsArena *arena)
+{
+  SymsArenaFrame frame = syms_arena_frame_begin(arena);
+
+  syms_uint buf_size = 8192;
+  char *buf = syms_arena_push_array(arena, char, buf_size);
+  SymsBuffer writer = syms_buffer_init(buf, buf_size);
+
+  SymsArenaFrame temp_frame = syms_arena_frame_begin(arena);
+  syms_bool type_formatted = syms_false;
+  const char *result = 0;
+
+  {
+    // format pointer and array
+    const char *att = 0;
+    syms_uint att_size = 4096;
+    char *att_buffer = syms_arena_push_array(arena, char, att_size);
+
+    SymsBuffer att_writer = syms_buffer_init(att_buffer, att_size);
+    SymsType base_type = *type;
+    SymsTypeModifier modifiers = base_type.modifier;
+
+    while (base_type.kind == SYMS_TYPE_PTR || base_type.kind == SYMS_TYPE_ARR) {  
+      SymsType next_type;
+      modifiers |= base_type.modifier;
+      if (base_type.kind == SYMS_TYPE_PTR) {
+        if (base_type.u.ptr.mode == SYMS_PTR_MODE_LVREF) {
+          syms_buffer_write_cstr(&att_writer, "&");
+        } else {
+          syms_buffer_write_cstr(&att_writer, "*");
+        }
+      } else if (base_type.kind == SYMS_TYPE_ARR) {
+        char arr[32];
+        syms_snprintf(arr, sizeof(arr), "[%u]", base_type.size);
+        syms_buffer_write_cstr(&att_writer, arr);
+      }
+      if (!syms_infer_type(instance, base_type.next_id, &next_type))
+        break;
+      if (next_type.kind == SYMS_TYPE_NULL)
+        break;
+      base_type = next_type;
+    }
+
+    switch (base_type.kind) {
+    case SYMS_TYPE_VARIADIC: {
+      syms_buffer_write_cstr(&writer, "...");
+    } break;
+
+    case SYMS_TYPE_PROC: {
+      struct SymsArenaFrame proc_frame;
+      SymsArgIter arg_iter;
+
+#if 0
+      const char *return_type_str;
+      proc_frame = syms_arena_frame_begin(arena);
+      return_type_str = syms_format_typeid(instance, &base_type.u.proc.ret_type_id, arena);
+      syms_buffer_write_cstr(&writer, return_type_str);
+      syms_buffer_write_cstr(&writer, " ");
+      syms_arena_frame_end(proc_frame);
+#endif
+
+      proc_frame = syms_arena_frame_begin(arena);
+
+      syms_buffer_write_cstr(&writer, "(*)(");
+      if (syms_arg_iter_init(instance, &arg_iter, &base_type)) {
+        SymsTypeID arg_id;
+        syms_uint arg_count = 0;
+
+        while (syms_arg_iter_next(&arg_iter, &arg_id)) {
+          struct SymsArenaFrame arg_frame = syms_arena_frame_begin(proc_frame.arena);
+          const char *arg_str;
+
+          if (syms_typeid_cmp(&base_type.id, &arg_id) == 0) {
+            syms_buffer_write_cstr(&writer, "???");
+            SYMS_ASSERT_FAILURE_PARANOID("self-referencing type detected");
+            continue;
+          }
+
+          arg_str = syms_format_typeid(instance, &arg_id, arg_frame.arena);
+          if (arg_count > 0) {
+            syms_buffer_write_cstr(&writer, ", ");
+          }
+          syms_buffer_write_cstr(&writer, arg_str);
+          arg_count += 1;
+          syms_arena_frame_end(arg_frame);
+        }
+        if (arg_count == 0) {
+          syms_buffer_write_cstr(&writer, "void");
+        }
+      }
+
+      syms_buffer_write_cstr(&writer, ")");
+      syms_arena_frame_end(proc_frame);
+    } break;
+    default: {
+      char temp[4096];
+      
+      // format type flags
+      syms_int flags = (syms_int)modifiers;
+      while (flags) {
+        syms_int flag_bit = flags & -flags;
+        const char *flag_str = syms_type_modifier_to_str((SymsTypeModifier)flag_bit);
+        syms_buffer_write_cstr(&writer, flag_str);
+        syms_buffer_write_cstr(&writer, " ");
+        flags &= ~flag_bit;
+      }
+
+      // Type name
+      temp[0] = 0;
+      syms_read_strref(instance, &base_type.name_ref, temp, sizeof(temp));
+      syms_buffer_write_cstr(&writer, temp);
+
+      // Array, pointer decorators
+      att = syms_buffer_get_cstr(&att_writer);
+      if (*att) {
+        syms_buffer_write_cstr(&writer, " ");
+        syms_buffer_write_cstr(&writer, att);
+      }
+    } break;
+    }
+
+    type_formatted = syms_true;
+  }
+  syms_arena_frame_end(temp_frame);
+
+  if (type_formatted) {
+    result = syms_buffer_get_cstr(&writer);
+    arena->size -= writer.size - writer.off;
+  } else {
+    syms_arena_frame_end(frame);
+  }
+  return result;
+}
+
+SYMS_API const char *
+syms_format_typeid(SymsInstance *instance, SymsTypeID *type_id, SymsArena *arena)
+{
+  const char *type_str = "<notype>";
+  SymsType type;
+  if (syms_infer_type(instance, *type_id, &type)) {
+    type_str = syms_format_type(instance, &type, arena);
+  }
+  return type_str;
+}
+
+SYMS_INTERNAL const char *
+syms_format_proc_args(SymsInstance *instance, SymsType *proc_type, SymsSymbolVar *locals, SymsArena *arena)
+{
+  struct SymsArenaFrame frame = syms_arena_frame_begin(arena);
+
+  syms_uint buffer_size = 8192;
+  char *buffer = syms_arena_push_array(arena, char, buffer_size);
+
+  const char *args_cstr = 0;
+
+  if (buffer) {
+    SymsBuffer writer = syms_buffer_init(buffer, buffer_size);
+
+    SymsArgIter arg_iter;
+    if (syms_arg_iter_init(instance, &arg_iter, proc_type)) {
+      syms_uint arg_idx = 0;
+      SymsTypeID arg_typeid;
+      while (syms_arg_iter_next(&arg_iter, &arg_typeid)) {
+        struct SymsArenaFrame arg_frame;
+        const char *type_str;
+        syms_uint local_idx;
+        SymsSymbolVar *l;
+
+        // Separate arguments with a ',' when there are more than one.
+        if (arg_idx > 0)
+          syms_buffer_write_cstr(&writer, ", ");
+
+        // Write argument type
+        arg_frame = syms_arena_frame_begin(arena);
+        type_str = syms_format_typeid(instance, &arg_typeid, arena);
+        syms_buffer_write_cstr(&writer, type_str);
+        syms_arena_frame_end(arg_frame);
+
+        // Find name of the argument
+        local_idx = 0;
+        for (l = locals; l != 0; l = l->next) {
+          if (l->flags & SYMS_VAR_FLAG_PARAM) {
+            if (syms_typeid_cmp(l->type_id, &arg_typeid) == 0) {
+              SYMS_ASSERT(l->flags & SYMS_VAR_FLAG_PARAM);
+              syms_buffer_write_cstr(&writer, " ");
+              syms_buffer_write_cstr(&writer, l->label);
+              break;
+            }
+          }
+
+          local_idx += 1;
+        }
+
+        arg_idx += 1;
+      }
+
+      // No arguments
+      if (arg_idx == 0)
+        syms_buffer_write_cstr(&writer, "void");
+    }
+
+    args_cstr = syms_buffer_get_cstr(&writer);
+  }
+
+  if (args_cstr == 0)
+    syms_arena_frame_end(frame);
+
+  return args_cstr;
+}
+
+SYMS_API const char *
+syms_format_proc(SymsInstance *instance, SymsSymbolProc *proc, SymsArena *arena)
+{
+  syms_uint buffer_size = 8192;
+  char *buffer = syms_arena_push_array(arena, char, buffer_size);
+  if (buffer) {
+    SymsBuffer writer = syms_buffer_init(buffer, buffer_size);
+    SymsType proc_type;
+    if (syms_infer_type(instance, *proc->type_id, &proc_type)) {
+      struct SymsArenaFrame frame;
+      const char *args_str;
+      SymsTypeID *return_type_id;
+
+      // Return type
+      switch (proc_type.kind) {
+      case SYMS_TYPE_PROC  : return_type_id = &proc_type.u.proc.ret_type_id; break;
+      case SYMS_TYPE_METHOD: return_type_id = &proc_type.u.method.ret_type_id; break;
+      default              : return_type_id = 0; break;
+      }
+      if (return_type_id) {
+        struct SymsArenaFrame ret_frame = syms_arena_frame_begin(arena);
+        const char *type_str = syms_format_typeid(instance, return_type_id, arena);
+        syms_buffer_write_cstr(&writer, type_str);
+        syms_arena_frame_end(ret_frame);
+      }
+      syms_buffer_write_cstr(&writer, " ");
+      // Procedure name
+      syms_buffer_write_cstr(&writer, proc->label);
+      // Arguments
+      frame = syms_arena_frame_begin(arena);
+      args_str = syms_format_proc_args(instance, &proc_type, proc->locals, arena);
+      syms_buffer_write_cstr(&writer, "(");
+      syms_buffer_write_cstr(&writer, args_str);
+      syms_buffer_write_cstr(&writer, ")");
+      syms_arena_frame_end(frame);
+    }
+    return syms_buffer_get_cstr(&writer);
+  }
+  return 0;
+}
+
+SYMS_API const char *
+syms_format_method(SymsInstance *instance, SymsTypeMember *member, SymsArena *arena)
+{
+  struct SymsArenaFrame temp_frame;
+  SymsType type;
+
+  syms_uint buffer_size = 8192;
+  char *buffer = syms_arena_push_array(arena, char, buffer_size);
+  SymsBuffer writer = syms_buffer_init(buffer, buffer_size);
+
+  if (syms_infer_type(instance, *member->type_id, &type)) {
+    if (type.kind == SYMS_TYPE_METHOD) {
+      const char *access_str   = syms_member_access_to_str(member->access);
+      const char *modifier_str = syms_member_modifier_to_str(member->modifier);
+
+      if (access_str) {
+        syms_buffer_write_cstr(&writer, access_str);
+        syms_buffer_write_cstr(&writer, " ");
+      }
+
+      if (modifier_str) {
+        syms_buffer_write_cstr(&writer, modifier_str);
+        syms_buffer_write_cstr(&writer, " ");
+      }
+
+      {
+        const char *return_type_str;
+        temp_frame = syms_arena_frame_begin(arena);
+        return_type_str = syms_format_typeid(instance, &type.u.method.ret_type_id, arena);
+        syms_buffer_write_cstr(&writer, return_type_str);
+        syms_buffer_write_cstr(&writer, " ");
+        syms_arena_frame_end(temp_frame);
+      }
+
+      syms_buffer_write_cstr(&writer, member->label);
+
+      {
+        SymsArgIter arg_iter;
+        syms_uint arg_count = 0;
+
+        syms_buffer_write_cstr(&writer, "(");
+        if (syms_arg_iter_init(instance, &arg_iter, &type)) {
+          SymsTypeID arg_id;
+
+          while (syms_arg_iter_next(&arg_iter, &arg_id)) {
+            const char *arg_str;
+
+            temp_frame = syms_arena_frame_begin(arena);
+            arg_str = syms_format_typeid(instance, &arg_id, arena);
+            if (arg_count > 0) {
+              syms_buffer_write_cstr(&writer, ", ");
+            }
+            syms_buffer_write_cstr(&writer, arg_str);
+            syms_arena_frame_end(temp_frame);
+            ++arg_count;
+          }
+        }
+        syms_buffer_write_cstr(&writer, ")");
+      }
+
+      if (member->modifier == SYMS_MEMBER_MODIFIER_PURE_VIRTUAL) {
+        syms_buffer_write_cstr(&writer, " = 0");
+      }
+    }
+  }
+
+  return syms_buffer_get_cstr(&writer);
+}
+
+SYMS_API const char *
+syms_format_symbol(SymsInstance *instance, SymsSymbol *sym, SymsArena *arena)
+{
+  switch (sym->kind) {
+  case SYMS_SYMBOL_PROC: return syms_format_proc(instance, &sym->u.proc, arena);
+  default: return NULL;
+  }
+}
+
+SYMS_API SymsSymbol *
+syms_symbol_from_addr(struct SymsInstance *instance, SymsAddr addr, struct SymsArena *arena)
+{
+  struct SymsArenaFrame arena_frame = syms_arena_frame_begin(arena);
+  // 
+  SymsSymbol  *sym_first = 0;
+  SymsSymbol **sym_list  = &sym_first;
+  SymsSymbol  *sym;
+  // 
+  SymsModID    nearest_mod    = SYMS_INVALID_MOD_ID;
+  SymsAddr          nearest_delta  = SYMS_ADDR_MAX;
+  SymsSourceFileMap nearest_source_map;
+
+  SymsModID mod_id;
+  for (mod_id = 0; mod_id < instance->mods_num; ++mod_id) {
+    SymsModInfo *mod = instance->mods + mod_id;
+    SymsSourceFileMap source_map;
+
+    // TODO(nick): 
+    //
+    //  it would be nice to export global and local data
+  
+    SymsRangeMap *map = syms_rangemap_search(mod->rangemap, addr);
+    if (map) {
+      SymsProc *proc;
+      SymsSymbolProc *proc_sym;
+      SymsInlineIter inline_iter;
+      SymsInlineSite site;
+
+      proc = syms_mod_info_find_proc(mod, map->id);
+      proc_sym = syms_arena_push_symbol_proc(sym_list, arena, proc->va, proc->len);
+      if (!proc_sym) break;
+
+      syms_export_proc(instance, proc, 0, proc_sym, arena);
+      proc_sym->locals = syms_export_locals(instance, proc, addr, &proc_sym->local_count, arena);
+
+      // Gather inline chain
+      if (syms_inline_iter_init(instance, &inline_iter, addr)) {
+        syms_uint site_count = 0;
+        while (syms_inline_iter_next(&inline_iter, &site)) {
+          SymsSymbolProc *site_sym = syms_arena_push_struct(arena, SymsSymbolProc);
+          if (!site_sym) break;
+
+          // export inline site symbol
+          syms_export_inline_site(instance, &site, site_sym, arena);
+          // add inline node to chain
+          site_sym->inline_chain = proc_sym->inline_chain;
+          proc_sym->inline_chain = site_sym;
+          // counter for storing
+          site_count += 1;
+
+          site_sym->locals = syms_export_locals(instance, &site.inlinee, addr, &site_sym->local_count, arena);
+        }
+      }
+    }
+
+    // collect source map
+    if (syms_line_table_map_va(&mod->line_table, addr, &source_map)) {
+      SymsAddr delta = addr - source_map.line.va;
+      if (delta < nearest_delta) {
+        nearest_delta      = delta;
+        nearest_mod        = mod_id;
+        nearest_source_map = source_map;
+      }
+    }
+  }
+
+  if (nearest_mod != SYMS_INVALID_MOD_ID) {
+    sym = syms_arena_push_symbol(sym_list, SYMS_SYMBOL_MODULE, arena);
+    if (sym) {
+      SymsModInfo *mod = instance->mods + nearest_mod;
+      SymsSecIter sec_iter;
+      sym->u.mod.name = syms_read_strref_arena(instance, &mod->header.name, arena);
+      if (syms_sec_iter_init(instance, &sec_iter)) {
+        SymsSection sec;
+        while (syms_sec_iter_next(&sec_iter, &sec)) {
+          if (sec.va >= addr && addr < sec.va + sec.data_size) {
+            u8 *data = syms_arena_push_array(arena, u8, sizeof(sec.impl));
+            if (!data)
+              break;
+            syms_memcpy(data, &sec.impl, sizeof(sec.impl));
+            break;
+          }
+        }
+      }
+    }
+  }
+
+  if (nearest_delta != SYMS_ADDR_MAX) {
+    sym = syms_arena_push_symbol(sym_list, SYMS_SYMBOL_SRCMAP, arena);
+    if (sym)
+      syms_export_srcmap(instance, &nearest_source_map, &sym->u.source_map, arena);
+  }
+
+  {
+    SymsAddr jump_addr;
+    if (syms_trampoline_from_ip(instance, addr, &jump_addr)) {
+      sym = syms_arena_push_symbol(sym_list, SYMS_SYMBOL_TRAMPOLINE, arena);
+      if (sym) {
+        sym->u.trampoline.jump_addr = jump_addr;
+      }
+    }
+  }
+
+
+  if (instance->debug_info.type == SYMS_FORMAT_PDB) {
+    SymsDebugInfoPdb *pdb_info = syms_get_debug_info_pdb(&instance->debug_info);
+    pdb_context *pdb = &pdb_info->context;
+    pdb_pointer name;
+    if (pdb_find_nearest_sym(pdb, addr, &name)) {
+      pdb_sc sc;
+      if (pdb_find_nearest_sc(pdb, addr, &sc)) {
+        SymsAddr start_addr;
+        if (pdb_build_va(pdb, sc.sec, sc.sec_off, &start_addr)) {
+          SymsProc proc;
+
+          syms_memzero(&proc, sizeof(proc));
+          proc.va = start_addr;
+          proc.len = sc.size;
+          proc.dbg_start_va = proc.va;
+          proc.dbg_end_va = proc.va + proc.len;
+          proc.name_ref = syms_string_ref_pdb(&name);
+
+          sym = syms_arena_push_symbol(sym_list, SYMS_SYMBOL_PROC, arena);
+          syms_export_proc(instance, &proc, 0, &sym->u.public_proc, arena);
+        }
+      }
+    }
+  }
+
+  if (arena->flags & SYMS_ARENA_FLAG_ALLOC_FAILED) {
+    syms_arena_frame_end(arena_frame);
+    sym_first = 0;
+  }
+
+  return sym_first;
+}
+
+/////
+////
+// Wrappers for source line query
+//
+
+SYMS_API SymsSymbolSrcmap *
+syms_src_to_va_arena(SymsInstance *instance, const char *path, syms_uint path_size, syms_uint ln, syms_uint col, SymsArena *arena)
+{ 
+  SymsSymbolSrcmap *map_out;
+  SymsSourceFileMap map;
+
+(void)col;
+
+  map_out = syms_arena_push_struct(arena, SymsSymbolSrcmap);
+  // TODO: cast fix!!!
+  if (syms_src_to_va(instance, (char *)path, path_size, ln, &map)) {
+    syms_export_srcmap(instance, &map, map_out, arena);
+  }
+  return map_out;
+}
+
+SYMS_API SymsSymbolSrcmap *
+syms_va_to_src_arena(SymsInstance *instance, SymsAddr va, SymsArena *arena)
+{
+  SymsSymbolSrcmap *map_out;
+  SymsSourceFileMap map;
+
+  map_out = syms_arena_push_struct(arena, SymsSymbolSrcmap);
+  if (syms_va_to_src(instance, va, &map)) {
+    syms_export_srcmap(instance, &map, map_out, arena);
+  }
+  return map_out;
+}
+
+//
+///
+/////
+
+SYMS_API SymsSymbol *
+syms_symbol_from_str_ex(SymsInstance *instance, const char *str_ptr, syms_uint str_size, struct SymsArena *arena)
+{
+  struct SymsArenaFrame arena_frame = syms_arena_frame_begin(arena);
+  SymsSymbol *sym_first = 0;
+  SymsSymbol **sym_list = &sym_first;
+
+  {
+    SymsProc proc;
+    SymsModID mod_id;
+    if (syms_proc_from_name_2(instance, str_ptr, str_size, &proc, &mod_id)) {
+      SymsSymbolProc *proc_sym = syms_arena_push_symbol_proc(sym_list, arena, proc.va, proc.len);
+      if (proc_sym) {
+        syms_export_proc(instance, &proc, 0, proc_sym, arena);
+      }
+    }
+  }
+
+  {
+    SymsConst cdata;
+    if (syms_const_from_name(instance, str_ptr, str_size, &cdata)) {
+      SymsSymbolConstData *cdata_sym = syms_arena_push_symbol_const_data(sym_list, arena);
+      syms_export_symbol_const_data(instance, &cdata, cdata_sym, arena);
+    }
+  }
+
+  {
+    SymsGlobal gdata;
+    if (syms_global_from_name(instance, str_ptr, str_size, &gdata)) {
+      SymsSymbolStaticData *gdata_sym = syms_arena_push_symbol_static_data(sym_list, arena);
+      syms_export_global_data(instance, &gdata, gdata_sym, arena);
+    }
+  }
+
+  {
+    SymsType type;
+    if (syms_type_from_name(instance, str_ptr, str_size, &type)) {
+      SymsSymbolTypeInfo *type_sym = syms_arena_push_symbol_type(sym_list, arena);
+      syms_export_type(instance, type.id, type_sym, arena);
+    }
+  }
+
+  if (arena->flags & SYMS_ARENA_FLAG_ALLOC_FAILED) {
+    syms_arena_frame_end(arena_frame);
+    sym_first = 0;
+  }
+
+  return sym_first;
+}
+
+SYMS_API SymsAddr
+syms_get_orig_rebase(const struct SymsInstance *instance)
+{
+  return instance->img.base_addr;
+}
+
+SYMS_API SymsAddr
+syms_get_rebase(const struct SymsInstance *instance)
+{
+  return instance->rebase;
+}
+
+SYMS_API void
+syms_set_rebase(struct SymsInstance *instance, SymsAddr rebase)
+{
+  SymsAddr prev_rebase = syms_get_rebase(instance);
+  SymsModID mod_id;
+  for (mod_id = 0; instance->mods_num; ++mod_id) {
+    SymsModInfo *mod = instance->mods + mod_id;
+    syms_uint i;
+
+    for (i = 0; i < mod->procs->push_count; ++i) {
+      SymsProc *p = syms_mod_info_find_proc(mod, i);
+      p->va           = (p->va - prev_rebase) + rebase;
+      p->dbg_start_va = (p->dbg_start_va - prev_rebase) + rebase;
+      p->dbg_end_va   = (p->dbg_end_va - prev_rebase) + rebase;
+    }
+
+    for (i = 0; i < mod->rangemap->push_count; ++i) {
+      SymsRangeMap *p = syms_mod_info_find_rangemap(mod, i);
+      p->lo = (p->lo - prev_rebase) + rebase;
+      p->hi = (p->hi - prev_rebase) + rebase;
+    }
+
+    for (i = 0; i < mod->line_table.line_count; ++i) {
+      SymsAddrMap *p = mod->line_table.addrs + i;
+      p->addr = (p->addr - prev_rebase) + rebase;
+    }
+  }
+
+  instance->rebase = rebase;
+}
+
+SYMS_API SymsImageType
+syms_get_image_type(const struct SymsInstance *instance)
+{
+  return instance->img.type;
+}
+
+SYMS_API syms_uint
+syms_get_image_size(const struct SymsInstance *instance)
+{
+  return (syms_uint)instance->img.img_data_size;
+}
+
+SYMS_API SymsImageHeaderClass
+syms_get_image_header_class(const struct SymsInstance *instance)
+{
+  return instance->img.header_class;
+}
+
+SYMS_API SymsFormatType
+syms_get_symbol_type(const struct SymsInstance *instance)
+{
+  return instance->debug_info.type;
+}
+
+SYMS_API syms_bindata_type
+syms_examine_binary_data(void *buf, SymsUMM buf_size)
+{
+  syms_bindata_type r = SYMS_BINDATA_TYPE_NULL;
+  SymsInstance *instance = syms_init();
+  do {
+    SymsErrorCode err = syms_load_image(instance, buf, buf_size, 0);
+    SymsFile file;
+
+    if (SYMS_RESULT_OK(err)) {
+      if (instance->img.type == SYMS_IMAGE_NT) {
+        static syms_bindata_type map[] = {
+          SYMS_BINDATA_TYPE_NULL,
+          SYMS_BINDATA_TYPE_NT32,
+          SYMS_BINDATA_TYPE_NT64
+        };
+        r = map[instance->img.header_class];
+        break;
+      }
+      if (instance->img.type == SYMS_IMAGE_ELF) {
+        static syms_bindata_type map[] = {
+          SYMS_BINDATA_TYPE_NULL,
+          SYMS_BINDATA_TYPE_ELF32,
+          SYMS_BINDATA_TYPE_ELF64
+        };
+        r = map[instance->img.header_class];
+        break;
+      }
+    }
+
+    file.path = "";
+    file.base = buf;
+    file.size = buf_size;
+    err = syms_load_debug_info_ex(instance, &file, 1, SYMS_LOAD_DEBUG_INFO_FLAGS_DEFER_BUILD_MODULE);
+    if (SYMS_RESULT_OK(err)) {
+      static syms_bindata_type map[] = {
+        SYMS_BINDATA_TYPE_NULL,
+        SYMS_BINDATA_TYPE_PDB,
+        SYMS_BINDATA_TYPE_DWARF,
+        SYMS_BINDATA_TYPE_ELFSYM
+      };
+      r = map[instance->debug_info.type];
+      break;
+    }
+  } while (0);
+  syms_quit(instance);
+
+  return r;
+}
+
+//
+//
+//
+
+SYMS_API SymsSymbolTypeInfo *
+syms_type_from_name_ex(SymsInstance *instance, const char *name, syms_uint name_size, SymsArena *arena)
+{
+  SymsType type;
+  SymsSymbolTypeInfo *type_out = NULL;
+  if (syms_type_from_name(instance, name, name_size, &type)) {
+    SymsErrorCode error_code;
+
+    type_out = syms_arena_push_struct(arena, SymsSymbolTypeInfo);
+    error_code = syms_export_type(instance, type.id, type_out, arena);
+    (void)error_code;
+    SYMS_ASSERT_PARANOID(SYMS_RESULT_OK(error_code));
+  }
+  return type_out;
+}
+
+SYMS_API SymsSymbolStaticData *
+syms_global_from_name_ex(SymsInstance *instance, const char *name, syms_uint name_size, SymsArena *arena)
+{
+  SymsGlobal gdata;
+  SymsSymbolStaticData *gdata_out = NULL;
+  if (syms_global_from_name(instance, name, name_size, &gdata)) {
+    gdata_out = syms_arena_push_struct(arena, SymsSymbolStaticData);
+    if (gdata_out) {
+      syms_export_global_data(instance, &gdata, gdata_out, arena);
+    }
+  }
+  return gdata_out;
+}
+
+SYMS_API SymsSymbolConstData *
+syms_const_from_name_ex(SymsInstance *instance, const char *name, syms_uint name_size, SymsArena *arena)
+{
+  SymsSymbolConstData *cdata_out = NULL;
+  SymsConst cdata;
+
+  if (syms_const_from_name(instance, name, name_size, &cdata)) {
+    cdata_out = syms_arena_push_struct(arena, SymsSymbolConstData);
+    if (cdata_out) {
+      syms_export_symbol_const_data(instance, &cdata, cdata_out, arena);
+    }
+  }
+
+  return cdata_out;
+}
+
+SYMS_API SymsSymbolProc *
+syms_proc_from_name_ex(SymsInstance *instance, const char *name, syms_uint name_size, SymsArena *arena)
+{
+  SymsSymbolProc *proc_out = NULL;
+  SymsProc proc;
+  SymsModID mod_id;
+
+  if (syms_proc_from_name_2(instance, name, name_size, &proc, &mod_id)) {
+    proc_out = syms_arena_push_struct(arena, SymsSymbolProc);
+    if (proc_out) {
+      syms_export_proc(instance, &proc, 0, proc_out, arena);
+    }
+  }
+
+  return proc_out;
+}
+
+SYMS_API SymsSymbolProc *
+syms_proc_from_addr_ex(SymsInstance *instance, SymsAddr addr, SymsArena *arena)
+{
+  SymsSymbolProc *proc_out = NULL;
+  SymsProc proc;
+  SymsModID mod_id;
+
+  if (syms_proc_from_va_2(instance, addr, &proc, &mod_id)) {
+    proc_out = syms_arena_push_struct(arena, SymsSymbolProc);
+    if (proc_out) {
+      syms_export_proc(instance, &proc, 0, proc_out, arena);
+    }
+  }
+
+  return proc_out;
+}
+
+//
+//
+//
+
+SYMS_API SymsArena *
+syms_borrow_memory(SymsInstance *instance)
+{
+  SymsArena arena;
+  syms_arena_init(&arena, syms_get_pagesize());
+  SymsArena *result = syms_arena_push_struct(&arena, SymsArena);
+  *result = arena;
+  if (instance) {
+    SymsBorrowedMemory *n = syms_arena_push_struct(instance->arena, SymsBorrowedMemory);
+    n->arena = result;
+    n->next = instance->borrowed_memory;
+    instance->borrowed_memory = n;
+  }
+  return result;
+}
+
+SYMS_API void
+syms_return_memory(SymsInstance *instance, SymsArena **arena_ptr)
+{
+  if (instance) {
+    for (SymsBorrowedMemory *n = instance->borrowed_memory, *p = 0; n != 0; p = n, n = n->next) {
+      if (n->arena == *arena_ptr) {
+        if (p) {
+          p->next = n->next;
+        } else {
+          instance->borrowed_memory = n->next;
+        }
+        break;
+      }
+    }
+  }
+  syms_arena_free(*arena_ptr);
+  *arena_ptr = 0;
+}
+
+SYMS_API SymsUMM
+syms_get_arena_size(struct SymsArena *arena)
+{
+  return arena->size;
+}
+
+SYMS_API SymsArenaFrame *
+syms_begin_arena_frame(struct SymsArena *arena)
+{
+  SymsUMM size = arena->size;
+  struct SymsArenaFrame *frame = syms_arena_push_struct(arena, struct SymsArenaFrame);
+  if (frame) {
+    frame->size = size;
+    frame->arena = arena;
+  }
+  return frame;
+}
+
+SYMS_API void
+syms_end_arena_frame(struct SymsArenaFrame *frame)
+{
+  syms_arena_frame_end(*frame);
+}
+
+//
+//
+//
+
+SYMS_API SymsRangeMap *
+syms_mod_info_find_rangemap(SymsModInfo *m, syms_uint i)
+{
+  return (SymsRangeMap *)syms_block_allocator_find_push(m->rangemap, i);
+}
+
+SYMS_API SymsRangeMap *
+syms_rangemap_search(SymsBlockAllocator *rangemap, SymsAddr va)
+{
+  SymsRangeMap *result = 0;
+  
+  if (rangemap->push_count > 0) {
+    s32 max = (s32)rangemap->push_count - 1;
+    s32 min = 0;
+    
+    while (min <= max) {
+      s32 mid = min + (max - min)/2;
+      SymsRangeMap *range;
+      SYMS_ASSERT(mid < (s32)rangemap->push_count);
+      range = (SymsRangeMap*)syms_block_allocator_find_push(rangemap, mid);
+      if (va < range->lo && va < range->hi) {
+        max = mid - 1;
+      } else if (va > range->lo && va > range->hi) {
+        min = mid + 1;
+      } else {
+        min = mid;
+        break;
+      }
+    }
+    
+    if (min >= 0 && (syms_uint)min < rangemap->push_count) {
+      for (u32 i = (u32)min + 1; i < rangemap->push_count; ++i) {
+        SymsRangeMap *range = (SymsRangeMap*)syms_block_allocator_find_push(rangemap, i);
+        if (va < range->lo || va >= range->hi) {
+          break;
+        }
+        min = (s32)i;
+      }
+      SYMS_ASSERT(min >= 0 && (u32)min < rangemap->push_count);
+      SymsRangeMap *range = (SymsRangeMap*)syms_block_allocator_find_push(rangemap, min);
+      if (va >= range->lo && va < range->hi) {
+        result = range;
+      }
+    }
+  }
+  
+  return result;
+}
+
+//
+//
+//
+
+SYMS_INTERNAL SymsErrorCode
+syms_memread(SymsMemread* info, SymsAddr va, void* buffer, syms_uint buffer_size)
+{
+    syms_uint read_len = info->callback(info->context, va, buffer, buffer_size);
+    if (read_len == buffer_size) {
+      info->result = SYMS_ERR_OK;
+    } else {
+      info->result = SYMS_ERR_INREAD;
+    }
+    return info->result;
+}
