RAD symbol library

diff --git a/Engine/Source/Runtime/SymsLib/src/syms.h b/Engine/Source/Runtime/SymsLib/src/syms.h
new file mode 100644
--- /dev/null
+++ b/Engine/Source/Runtime/SymsLib/src/syms.h
@@ -0,0 +1,110 @@
+// Copyright Epic Games, Inc. All Rights Reserved.
+#ifndef SYMS_INCLUDE_H
+#define SYMS_INCLUDE_H
+
+/******************************************************************************
+ * File   : syms.h                                                            *
+ * Author : Nikita Smith                                                      *
+ * Created: 2019/02/20                                                        *
+ * Purpose: Provides an API for accessing debug information                   *
+ ******************************************************************************/
+
+typedef struct SymsBorrowedMemory
+{
+  SymsArena *arena;
+  struct SymsBorrowedMemory *next;
+} SymsBorrowedMemory;
+
+struct SymsInstance
+{
+  syms_uint user_id;
+  SymsArena *arena;
+  SymsImage img;
+  SymsDebugInfo debug_info;
+  SymsMod null_mod;
+  SymsAddr rebase; // base address that is used for calculating addresses of symbols, lines, and etc.
+  syms_uint mods_num;
+  SymsModInfo *mods;
+  SymsBorrowedMemory *borrowed_memory;
+  struct 
+  {
+    SymsErrorCode code;
+    SymsString text;
+  } error;
+};
+
+
+typedef union SymsProcData
+{
+  struct 
+  {
+    pdb_cvdata_token cvdata;
+  } pdb;
+  struct 
+  {
+    DwTag type_tag;
+    DwAttrib frame_base;
+    SymsOffset range_off;
+  } dw;
+} SymsProcData;
+SYMS_COMPILER_ASSERT(sizeof(SymsProcData) <= sizeof(SymsProcImpl));
+
+typedef struct SymsMemread 
+{
+  SymsErrorCode result;
+  void *context;
+  syms_memread_sig *callback;
+} SymsMemread;
+
+typedef struct
+{
+  SymsErrorCode result;
+  SymsArch arch;
+  void *context;
+  syms_regread_sig *callback;
+} SymsRegread;
+
+typedef struct
+{
+  SymsErrorCode result;
+  SymsArch arch;
+  void *context;
+  syms_regwrite_sig *callback;
+} SymsRegwrite;
+
+SYMS_INTERNAL SymsTypeID
+syms_type_id_null(void);
+
+SYMS_INTERNAL void
+syms_line_init(SymsLine *line, SymsAddr va, syms_uint ln, syms_uint col);
+
+SYMS_API SymsModInfo *
+syms_get_mod(SymsInstance *instance, SymsModID mod_id);
+
+SYMS_API SymsModID
+syms_infer_global_data_module(SymsInstance *instance, SymsGlobal *gdata);
+// Resolve module that contains global-data.
+
+typedef struct SymsRangeMap
+{
+  SymsAddr lo;
+  SymsAddr hi;
+  syms_uint id;
+} SymsRangeMap;
+
+SYMS_API SymsRangeMap *
+syms_mod_info_find_rangemap(SymsModInfo *m, syms_uint i);
+
+SYMS_API SymsRangeMap *
+syms_rangemap_search(SymsBlockAllocator *rangemap, SymsAddr va);
+
+SYMS_INTERNAL SymsErrorCode
+syms_memread(SymsMemread* info, SymsAddr va, void* buffer, syms_uint buffer_size);
+SYMS_API SymsArenaFrame *
+syms_begin_arena_frame(struct SymsArena *arena);
+
+SYMS_API void
+syms_end_arena_frame(struct SymsArenaFrame *frame);
+
+#endif // SYMS_INCLUDE_H
+
diff --git a/Engine/Source/Runtime/SymsLib/src/syms_block_alloc.c b/Engine/Source/Runtime/SymsLib/src/syms_block_alloc.c
new file mode 100644
--- /dev/null
+++ b/Engine/Source/Runtime/SymsLib/src/syms_block_alloc.c
@@ -0,0 +1,166 @@
+// Copyright Epic Games, Inc. All Rights Reserved.
+/******************************************************************************
+ * File   : syms_block_alloc.c                                                *
+ * Author : Nikita Smith                                                      *
+ * Created: 2021/02/21                                                        *
+ ******************************************************************************/
+
+SYMS_API SymsBlockAllocator *
+syms_block_allocator_init(u32 push_size, u32 block_size, SymsArena *arena)
+{
+  SymsBlockAllocator *ba = syms_arena_push_struct(arena, SymsBlockAllocator);
+  ba->push_size   = push_size;
+  ba->push_count  = 0;
+  ba->block_size  = block_size;
+  ba->block_count = 0;
+  ba->last_block_size = block_size;
+  ba->first_block = 0;
+  ba->next_block  = &ba->first_block;
+  ba->arena       = arena;
+  ba->index_table = 0;
+  return ba;
+}
+
+SYMS_API void *
+syms_block_allocator_push(SymsBlockAllocator *a)
+{
+  SymsBlockNode *b;
+  if (a->last_block_size + a->push_size > a->block_size) {
+    b = syms_arena_push_struct(a->arena, SymsBlockNode);
+    b->base = syms_arena_push_array(a->arena, char, a->block_size);
+    b->next = 0;
+
+    *a->next_block = b;
+    a->next_block = &b->next;
+
+    a->last_block = b;
+    a->last_block_size = 0;
+
+    ++a->block_count;
+  } else {
+    b = a->last_block;
+  }
+  void *ptr = (char*)a->last_block->base + a->last_block_size;
+  a->last_block_size += a->push_size;
+  a->push_count += 1;
+  return ptr;
+}
+
+SYMS_API void
+syms_block_allocator_build_index_table(SymsBlockAllocator *a)
+{
+  SYMS_ASSERT(a->index_table == 0);
+  a->index_table = syms_arena_push_array(a->arena, SymsBlockNode *, a->block_count);
+  SymsBlockNode *b = a->first_block; 
+  for (u32 i = 0; i < a->block_count; ++i) {
+    a->index_table[i] = b;
+    b = b->next;
+  }
+  SYMS_ASSERT(b == 0);
+}
+
+SYMS_API void *
+syms_block_allocator_find_push(SymsBlockAllocator *a, u32 pi)
+{
+  SYMS_ASSERT(pi < a->push_count);
+  u32 push_offset = pi * a->push_size;
+  u32 block_index = push_offset / a->block_size;
+  SymsBlockNode *block = a->index_table[block_index];
+  u32 base_offset = push_offset - (block_index * a->block_size);
+  SYMS_ASSERT(base_offset + a->push_size <= a->block_size);
+  void *p = (void*)(block->base + base_offset);
+  return p;
+}
+
+SYMS_API void
+syms_block_allocator_swap(SymsBlockAllocator *a, u32 l, u32 r)
+{
+  SymsArenaFrame frame = syms_arena_frame_begin(a->arena);
+  u8 *t = syms_arena_push_array(frame.arena, u8, a->push_size);
+  void *pl = syms_block_allocator_find_push(a, (u32)l);
+  void *pr = syms_block_allocator_find_push(a, (u32)r);
+  syms_memcpy(t, pl, a->push_size);
+  syms_memcpy(pl, pr, a->push_size);
+  syms_memcpy(pr, t, a->push_size);
+  syms_arena_frame_end(frame);
+}
+
+SYMS_INTERNAL void
+syms_block_allocator_ins_sort(SymsBlockAllocator *allocator, s32 n, s32 (*compare)(const void *a, const void *b))
+{
+  SymsArenaFrame frame = syms_arena_frame_begin(allocator->arena);
+  u8 *t = syms_arena_push_array(allocator->arena, u8, allocator->push_size);
+  for (s32 i = 1; i < n; ++i) {
+    s32 j = i;
+    
+    void *a = syms_block_allocator_find_push(allocator, i);
+    syms_memcpy(t, a, allocator->push_size);
+    a = t;
+
+    while (j > 0) {
+      void *b = syms_block_allocator_find_push(allocator, j-1);
+      s32 c = compare(a, b);
+      if (c == 0) break;
+      syms_block_allocator_swap(allocator, j, j-1);
+      --j;
+    }
+    if (i != j) {
+      void *dst = syms_block_allocator_find_push(allocator, j);
+      syms_memcpy(dst, t, allocator->push_size);
+    }
+  }
+  syms_arena_frame_end(frame);
+}
+
+SYMS_INTERNAL void
+syms_block_allocator_quicksort(SymsBlockAllocator *allocator, s32 w, s32 n, s32 (*compare)(const void *a, const void *b))
+{
+  while (n > 12) {
+    s32 m = n >> 1;
+    s32 c01 = compare(syms_block_allocator_find_push(allocator, w + 0), syms_block_allocator_find_push(allocator, w + m));
+    s32 c12 = compare(syms_block_allocator_find_push(allocator, w + m), syms_block_allocator_find_push(allocator, w + n-1));
+    if (c01 != c12) {
+      s32 c = compare(syms_block_allocator_find_push(allocator, w + 0), syms_block_allocator_find_push(allocator, w + n-1));
+      s32 z = (c == c12) ? 0 : n-1;
+      syms_block_allocator_swap(allocator, w + z, w + m);
+    }
+    syms_block_allocator_swap(allocator, w + 0, w + m);
+    s32 i = 1;
+    s32 j = n-1;
+    for (;;) {
+      for (;;++i) {
+        void *a = syms_block_allocator_find_push(allocator, w + i);
+        void *b = syms_block_allocator_find_push(allocator, w + 0);
+        if (!compare(a,b)) break;
+      }
+      for (;;--j) {
+        void *a = syms_block_allocator_find_push(allocator, w + 0);
+        void *b = syms_block_allocator_find_push(allocator, w + j);
+        if (!compare(a,b)) break;
+      }
+      if (i >= j) break;
+      syms_block_allocator_swap(allocator, w + i, w + j);
+      ++i;
+      --j;
+    }
+    if (j < (n-i)) {
+      syms_block_allocator_quicksort(allocator, w + 0, j, compare);
+      w += i;
+      n = n-i;
+    } else {
+      syms_block_allocator_quicksort(allocator, w + i, n - i, compare);
+      n = j;
+    }
+  }
+}
+
+SYMS_API void
+syms_block_allocator_sort(SymsBlockAllocator *a, s32 (*compare)(const void *a, const void *b))
+{
+  if (a->push_count > 1) {
+    // NOTE: butchered quicksort from stb.h
+    syms_block_allocator_quicksort(a, 0, a->push_count, compare);
+    syms_block_allocator_ins_sort(a, a->push_count, compare);
+  }
+}
+
diff --git a/Engine/Source/Runtime/SymsLib/src/syms_block_alloc.h b/Engine/Source/Runtime/SymsLib/src/syms_block_alloc.h
new file mode 100644
--- /dev/null
+++ b/Engine/Source/Runtime/SymsLib/src/syms_block_alloc.h
@@ -0,0 +1,40 @@
+// Copyright Epic Games, Inc. All Rights Reserved.
+#ifndef SYMS_BLOCK_ALLOC_H
+#define SYMS_BLOCK_ALLOC_H
+
+typedef struct SymsBlockNode
+{
+  char *base;
+  struct SymsBlockNode *next;
+} SymsBlockNode;
+
+struct SymsBlockAllocator
+{
+  syms_uint push_size;
+  syms_uint push_count;
+  syms_uint block_size;
+  syms_uint block_count;
+  syms_uint last_block_size;
+  SymsBlockNode *first_block;
+  SymsBlockNode **next_block;
+  SymsBlockNode *last_block;
+  SymsArena *arena;
+  SymsBlockNode **index_table;
+};
+
+SYMS_API SymsBlockAllocator *
+syms_block_allocator_init(syms_uint push_size, syms_uint block_size, SymsArena *arena);
+
+SYMS_API void *
+syms_block_allocator_push(SymsBlockAllocator *a);
+
+SYMS_API void
+syms_block_allocator_build_index_table(SymsBlockAllocator *a);
+
+SYMS_API void *
+syms_block_allocator_find_push(SymsBlockAllocator *a, syms_uint pi);
+
+SYMS_API void
+syms_block_allocator_sort(SymsBlockAllocator *a, syms_int (*compare)(const void *a, const void *b));
+
+#endif // SYMS_BLOCK_ALLOC_H
diff --git a/Engine/Source/Runtime/SymsLib/src/syms_core.h b/Engine/Source/Runtime/SymsLib/src/syms_core.h
new file mode 100644
--- /dev/null
+++ b/Engine/Source/Runtime/SymsLib/src/syms_core.h
@@ -0,0 +1,1284 @@
+// Copyright Epic Games, Inc. All Rights Reserved.
+#ifndef SYMS_CORE_INCLUDE_H
+#define SYMS_CORE_INCLUDE_H
+
+/******************************************************************************
+ * File   : syms_core.h                                                       *
+ * Author : Nikita Smith                                                      *
+ * Created: 2020/06/20                                                        *
+ * Purpose: core utilities for pdb and dwarf parsers                          *
+ ******************************************************************************/
+
+#define SYMS_PARANOID
+
+#ifndef syms_malloc
+#include <stdlib.h>
+#define syms_malloc(size, ud) ((void)ud,malloc(size))
+#define syms_free(ptr, ud)    ((void)ud,free(ptr))
+#endif
+#define syms_malloc_array(type, count, ud) (type *)syms_malloc(sizeof(type) * (count), ud)
+#define syms_malloc_struct(type, ud)       syms_malloc_array(type, 1, ud)
+
+#ifndef syms_qsort
+#include <stdlib.h>
+#define syms_qsort qsort
+#endif
+
+#ifndef syms_memcpy
+#include <string.h>
+#define syms_memcpy memcpy
+#endif
+
+#ifndef syms_memset
+#include <string.h>
+#define syms_memset memset
+#endif
+
+#ifndef syms_memcmp
+#include <string.h>
+#define syms_memcmp memcmp
+#endif
+
+#ifndef syms_strlen
+#include <string.h>
+#define syms_strlen strlen
+#endif
+
+#ifndef syms_strcmp
+#include <string.h>
+#define syms_strcmp strcmp
+#endif
+
+#ifndef syms_sprintf
+#include <stdio.h>
+#define syms_sprintf   sprintf
+#define syms_snprintf  snprintf
+#define syms_vsprintf  vsprintf
+#define syms_vsnprintf vsnprintf
+#endif
+
+#define syms_memzero(dst, size) syms_memset(dst, 0, size)
+
+#ifndef SYMS_INTERNAL
+#define SYMS_INTERNAL static
+#endif
+
+#ifndef SYMS_API
+#define SYMS_API
+#endif
+
+#ifndef SYMS_INLINE
+#define SYMS_INLINE inline
+#endif
+
+#ifndef SYMS_CUSTOM_TYPES
+#include <stdint.h>
+typedef uint8_t U8;
+typedef uint16_t U16;
+typedef uint32_t U32;
+typedef int8_t S8;
+typedef int16_t S16;
+typedef int32_t S32;
+typedef float F32;
+typedef double F64;
+typedef uint64_t U64;
+typedef int64_t S64;
+#endif
+
+typedef U8  u8;
+typedef U16 u16;
+typedef U32 u32;
+typedef U64 u64;
+typedef S8  s8;
+typedef S16 s16;
+typedef S32 s32;
+typedef S64 s64;
+typedef F32 f32;
+typedef F64 f64;
+
+typedef S32 syms_bool;
+
+#define syms_false 0
+#define syms_true 1
+
+#define SYMS_INT8_MIN  0x80 // -128
+#define SYMS_INT8_MAX  0x7f // +127
+#define SYMS_INT16_MIN 0x8000 // -32768
+#define SYMS_INT16_MAX 0x7fff // +32767
+#define SYMS_INT32_MIN 0x80000000 // -2147483648
+#define SYMS_INT32_MAX 0x7fffffff // +2147483647
+#define SYMS_INT64_MIN 0x8000000000000000ull // -9223372036854775808
+#define SYMS_INT64_MAX 0x7fffffffffffffffull // +9223372036854775807
+
+#define SYMS_UINT8_MAX  0xFFu
+#define SYMS_UINT16_MAX 0xFFFFu
+#define SYMS_UINT32_MAX 0xFFFFFFFFu
+#define SYMS_UINT64_MAX 0xFFFFFFFFFFFFFFFFu
+
+#define SYMS_TOSTR_X(x) #x
+#define SYMS_TOSTR(x) SYMS_TOSTR_X(x)
+#define SYMS_JOIN(a,b) a ## b
+
+#ifndef SYMS_ASSERT
+#ifdef _WIN32
+#include <intrin.h>
+#define SYMS_ASSERT(x) do { if (!(x)) __debugbreak(); } while (0)
+#else
+#include <assert.h>
+#define SYMS_ASSERT(x) assert(x)
+#endif
+#endif
+
+#ifndef SYMS_ASSERT_ALWAYS
+#define SYMS_ASSERT_ALWAYS(x) SYMS_ASSERT(x)
+#endif
+
+#ifndef SYMS_ASSERT_FAILURE
+#define SYMS_ASSERT_FAILURE(x) SYMS_ASSERT(0)
+#endif
+
+#ifndef SYMS_ASSERT_FAILURE_ALWAYS
+#define SYMS_ASSERT_FAILURE_ALWAYS(x) SYMS_ASSERT(0)
+#endif
+
+#if defined(SYMS_PARANOID)
+#define SYMS_ASSERT_PARANOID(x)           SYMS_ASSERT(x)
+#define SYMS_ASSERT_FAILURE_PARANOID(msg) SYMS_ASSERT_FAILURE(msg)
+#else
+#define SYMS_ASSERT_PARANOID(x)
+#define SYMS_ASSERT_FAILURE_PARANOID(msg)
+#endif
+
+#define SYMS_NOT_IMPLEMENTED    SYMS_ASSERT_FAILURE_ALWAYS("not implemented")
+#define SYMS_INVALID_CODE_PATH  SYMS_ASSERT_FAILURE_ALWAYS("invalid code path")
+#define SYMS_ASSERT_NO_SUPPORT  SYMS_INVALID_CODE_PATH
+
+#ifndef SYMS_COMPILER_ASSERT
+#define SYMS_COMPILER_ASSERT(predicate) SYMS_COMPILER_ASSERT_IMPL(predicate, __LINE__)
+#define SYMS_COMPILER_ASSERT_IMPL(predicate, line) typedef char SYMS_JOIN(compiler_dummy_name, line)[(predicate) ? +1 : -1]
+#endif
+
+#ifndef SYMS_ARRAY_SIZE
+#define SYMS_ARRAY_SIZE(x) (sizeof(x)/sizeof((x)[0]))
+#endif
+
+#ifndef SYMS_MIN
+#define SYMS_MIN(a,b) ((a) < (b) ? (a) : (b))
+#endif
+
+#ifndef SYMS_MAX
+#define SYMS_MAX(a,b) ((a) > (b) ? (a) : (b))
+#endif
+
+#ifndef SYMS_ABS
+#define SYMS_ABS(x) ((x) < 0 ? -(x) : (x))
+#endif
+
+#ifndef SYMS_MEMBER_OFFSET
+#define SYMS_MEMBER_OFFSET(type, member) offsetof(type, member)
+#endif
+
+#ifndef SYMS_GET_MEMBER_SIZE
+#define SYMS_GET_MEMBER_SIZE(type, member) sizeof((((type*)0)->member))
+#endif
+
+#ifdef _MSC_VER
+#define SYMS_PRIu64 "I64u"
+#define SYMS_PRId64 "I64d"
+#define SYMS_PRIx64 "I64x"
+#else
+#include <inttypes.h>
+#define SYMS_PRIu64 PRIu64
+#define SYMS_PRId64 PRId64
+#define SYMS_PRIx64 PRIx64
+#endif
+
+#define SYMS_KB(num) ((num)*1024u)
+#define SYMS_MB(num) (SYMS_KB(num)*1024u)
+#define SYMS_GB(num) (SYMS_MB(num)*1024u)
+#define SYMS_TB(num) (SYMS_TB(num)*1024u)
+
+#if defined(__amd64__) || defined(__amd64) || defined(__x86_64__) || defined(_M_X64) || defined(_M_AMD64)
+#define SYMS_X86
+#define SYMS_X64
+#define SYMS_LITTLE_ENDIAN
+#elif defined(i386) || defined(__i386) || defined(__i386__) || defined(_X86_) || defined(_M_I86)
+#define SYMS_X86
+#define SYMS_LITTLE_ENDIAN
+#elif defined(__arm__) || defined(_M_ARM) 
+#define SYMS_ARM
+#define SYMS_LITTLE_ENDIAN
+#elif defined(__aarch64__)
+#define SYMS_ARM64
+#define SYMS_LITTLE_ENDIAN
+#else
+#error "arch not defined"
+#endif
+
+#ifndef SYMS_CUSTOM_BSWAP
+#if defined(__clang__)
+#define syms_bswap16 __builtin_bswap16
+#define syms_bswap32 __builtin_bswap32
+#define syms_bswap64 __builtin_bswap64
+#elif defined(_MSC_VER)
+#define syms_bswap16 _byteswap_ushort
+#define syms_bswap32 _byteswap_ulong
+#define syms_bswap64 _byteswap_uint64
+#else
+SYMS_INTERNAL U16 syms_bswap16(U16 x) 
+{
+    return (((x >> 8) & 0xffu) | ((x & 0xffu) << 8));
+}
+SYMS_INTERNAL U32 syms_swap32(U32 x) 
+{
+    return (((x & 0xff000000u) >> 24) |
+            ((x & 0x00ff0000u) >> 8) |
+            ((x & 0x0000ff00u) << 8) |
+            ((x & 0x000000ffu) << 24));
+}
+SYMS_INTERNAL U64 syms_bswap64(U64 x) 
+{
+    return (((x & 0xff00000000000000) >> 56) |
+            ((x & 0x00ff000000000000) >> 40) |
+            ((x & 0x0000ff0000000000) >> 24) |
+            ((x & 0x000000ff00000000) >> 8) |
+            ((x & 0x00000000ff000000) << 8) |
+            ((x & 0x0000000000ff0000) << 24) |
+            ((x & 0x000000000000ff00) << 40) |
+            ((x & 0x00000000000000ff) << 56));
+}
+#endif
+#endif
+
+#if defined(SYMS_X64) || defined(SYMS_ARM64)
+typedef U64 SymsUWord;
+typedef S64 SymsSWord;
+#define SYMS_UWORD_MAX SYMS_UINT64_MAX
+#define SYMS_SWORD_MAX SYMS_INT64_MAX
+#define syms_bswapaddr(x) syms_bswap64(x)
+#elif defined(SYMS_X86)
+typedef U32 SymsUWord;
+typedef S32 SymsSWord;
+#define SYMS_UWORD_MAX SYMS_UINT32_MAX
+#define SYMS_SWORD_MAX SYMS_INT32_MAX
+#define syms_bswapaddr(x) syms_bswap32(x)
+#else
+#error "Address is not defined for this architecture"
+#endif
+
+//typedef SymsUMM SymsAddr;
+
+typedef enum 
+{
+    SYMS_ADDR_SIZE_NULL,
+    SYMS_ADDR_SIZE_16,
+    SYMS_ADDR_SIZE_32,
+    SYMS_ADDR_SIZE_64,
+    SYMS_ADDR_SIZE_128,
+    
+    SYMS_ADDR_SIZE_HOST =
+#if SYMS_ADDR_MAX == SYMS_UINT16_MAX
+    SYMS_ADDR_SIZE_16
+#elif SYMS_ADDR_MAX == SYMS_UINT32_MAX
+    SYMS_ADDR_SIZE_32
+#elif SYMS_ADDR_MAX == SYMS_UINT64_MAX 
+    SYMS_ADDR_SIZE_64
+#elif SYMS_ADDR_MAX == SYMS_UINT64_MAX*2
+    SYMS_ADDR_SIZE_128
+#else
+#error "Address size is not defined"
+#endif
+} SymsAddrSize;
+
+typedef enum 
+{
+    SYMS_ENDIAN_NULL,
+    SYMS_ENDIAN_LITTLE,
+    SYMS_ENDIAN_BIG,
+    
+    SYMS_ENDIAN_HOST = SYMS_ENDIAN_LITTLE
+} SymsEndian;
+
+#if defined(SYMS_LITTLE_ENDIAN)
+#define SYMS_LE16(x)     (x)
+#define SYMS_LE32(x)     (x)
+#define SYMS_LE64(x)     (x)
+#define SYMS_LEADDR(x)   (x)
+#define SYMS_BE16(x)     syms_bswap16(x)
+#define SYMS_BE32(x)     syms_bswap32(x)
+#define SYMS_BE64(x)     syms_bswap64(x)
+#define SYMS_BEADDR(x)   syms_bswapaddr(x)
+#else
+#define SYMS_LE16(x)     syms_bswap16(x)
+#define SYMS_LE32(x)     syms_bswap32(x)
+#define SYMS_LE64(x)     syms_bswap64(x)
+#define SYMS_LEADDR(x)   syms_bswapaddr(x)
+#define SYMS_BE16(x)     (x)
+#define SYMS_BE32(x)     (x)
+#define SYMS_BE64(x)     (x)
+#define SYMS_BEADDR(x)   (x)
+#endif
+
+#define SYMS_OFFSET_INVALID SYMS_UWORD_MAX
+typedef SymsUWord SymsOffset;
+
+// -------------------------------------------------------------------------------- 
+
+typedef struct SymsBlock
+{
+    SymsUMM len;
+    void *data;
+} SymsBlock;
+
+// -------------------------------------------------------------------------------- 
+
+#define SYMS_STRING_LIT(x) { sizeof(x) - 1, x }
+#define syms_string_lit(x) syms_string_init(x, sizeof(x) - 1)
+
+typedef struct SymsStringArray 
+{
+    U32 count;
+    SymsString *entries;
+} SymsStringArray;
+
+SYMS_API SymsString
+syms_string_init(const char *data, U32 len);
+
+SYMS_API SymsString
+syms_string_init_lit(const char *data);
+
+SYMS_API syms_bool
+syms_string_is_null(SymsString str);
+
+SYMS_API syms_bool
+syms_string_cmp_lit(SymsString a, const char *b);
+
+SYMS_API syms_bool
+syms_is_alpha_ascii(U32 codepoint);
+
+SYMS_API U32
+syms_lowercase(U32 codepoint);
+
+SYMS_API U32
+syms_uppercase(U32 codepoint);
+
+SYMS_API syms_bool
+syms_string_cmp_nocase(SymsString a_str, SymsString b_str);
+
+SYMS_API char
+syms_string_peek_byte(SymsString str, U32 byte_offset);
+
+SYMS_API U32
+syms_string_get_size(SymsString str);
+
+SYMS_API syms_bool
+syms_string_cmp(SymsString a, SymsString b);
+
+SYMS_API const U8 *
+syms_decode_utf8(const U8 *p, U32 *dst);
+
+SYMS_API char *
+syms_string_to_cstr(SymsString str, SymsArena *arena);
+
+
+// -------------------------------------------------------------------------------- 
+
+typedef struct SymsBuffer 
+{
+    const void *base;
+    SymsUMM size;
+    SymsUMM off;
+    u32 addr_size;
+    u32 offs_size;
+} SymsBuffer;
+
+SYMS_API SymsBuffer
+syms_buffer_init(const void *data, SymsUMM size);
+
+SYMS_API SymsBuffer
+syms_buffer_init_ex(const void *data, SymsUMM size, SymsAddrSize addr_size, SymsAddrSize offset_type);
+
+SYMS_API syms_bool
+syms_buffer_skip(SymsBuffer *buffer, SymsOffset num);
+
+SYMS_API syms_bool
+syms_buffer_seek(SymsBuffer *buffer, SymsOffset new_off);
+
+SYMS_API U8
+syms_buffer_peek_u8(SymsBuffer *bin);
+
+SYMS_API U16
+syms_buffer_peek_u16(SymsBuffer *bin);
+
+SYMS_API U32
+syms_buffer_peek_u32(SymsBuffer *bin);
+
+SYMS_API U64
+syms_buffer_peek_u64(SymsBuffer *bin);
+
+SYMS_API S8
+syms_buffer_peek_s8(SymsBuffer *bin);
+
+SYMS_API S16
+syms_buffer_peek_s16(SymsBuffer *bin);
+
+SYMS_API S32
+syms_buffer_peek_s32(SymsBuffer *bin);
+
+SYMS_API S64
+syms_buffer_peek_s64(SymsBuffer *bin);
+
+SYMS_API S64
+syms_buffer_peek_sleb128(SymsBuffer *bin);
+
+SYMS_API U64
+syms_buffer_peek_uleb128(SymsBuffer *bin);
+
+SYMS_API SymsString
+syms_buffer_peek_string(SymsBuffer *bin);
+
+SYMS_API SymsString
+syms_buffer_peek_cstr(SymsBuffer *bin);
+
+#define syms_buffer_push_array(buffer, type, count) (type *)syms_buffer_push(buffer, sizeof(type) * (count))
+#define syms_buffer_push_struct(buffer, type)       syms_buffer_push_array(buffer, type, 1)
+SYMS_API void *
+syms_buffer_push(SymsBuffer *buffer, SymsUMM num_bytes);
+
+SYMS_API syms_bool
+syms_buffer_read(SymsBuffer *buffer, void *out_bf, SymsUMM size);
+
+SYMS_API U8
+syms_buffer_read_u8(SymsBuffer *buffer);
+
+SYMS_API U16
+syms_buffer_read_u16(SymsBuffer *buffer);
+
+SYMS_API U32
+syms_buffer_read_u32(SymsBuffer *buffer);
+
+SYMS_API U32
+syms_buffer_read_u32(SymsBuffer *buffer);
+
+SYMS_API U64
+syms_buffer_read_u64(SymsBuffer *buffer);
+
+SYMS_API U64
+syms_buffer_read_uleb128(SymsBuffer *bin);
+// Reads an unsigned integer with variable-length encoding LEB128.
+
+SYMS_API S8
+syms_buffer_read_s8(SymsBuffer *buffer);
+
+SYMS_API S16
+syms_buffer_read_s16(SymsBuffer *buffer);
+
+SYMS_API S32
+syms_buffer_read_s32(SymsBuffer *buffer);
+
+SYMS_API S64
+syms_buffer_read_s64(SymsBuffer *buffer);
+
+SYMS_API S64
+syms_buffer_read_sleb128(SymsBuffer *bin);
+
+SYMS_API SymsString
+syms_buffer_read_string(SymsBuffer *buffer);
+
+SYMS_API SymsString
+syms_buffer_read_cstr(SymsBuffer *buffer);
+
+SYMS_API void
+syms_buffer_write_cstr(SymsBuffer *buffer, const char *cstr);
+
+SYMS_API void
+syms_buffer_write(SymsBuffer *buffer, const void *src, U32 src_size);
+
+SYMS_API void
+syms_buffer_write_u8(SymsBuffer *buffer, U8 value);
+
+SYMS_API const char *
+syms_buffer_get_cstr(SymsBuffer *buffer);
+
+SYMS_API const void *
+syms_buffer_at(SymsBuffer *buffer);
+
+// -------------------------------------------------------------------------------- 
+
+SYMS_API SymsString
+syms_path_get_file_name(SymsString path, syms_bool strip_ext);
+
+// -------------------------------------------------------------------------------- 
+
+enum 
+{
+    SYMS_ARENA_FLAG_ALLOC_FAILED = (1 << 0)
+};
+typedef u32 SymsArenaFlags;
+
+struct SymsArena 
+{
+    SymsArenaFlags flags;
+    SymsUMM size;
+    SymsUMM page_size;
+    void *head;
+    u8 *cursor;
+};
+
+struct SymsArenaFrame
+{
+    SymsArena *arena;
+    SymsArenaFlags flags;
+    SymsUMM size;
+    void *head;
+    u8 *cursor;
+};
+
+SYMS_API void
+syms_arena_init(SymsArena *arena, SymsUMM page_size);
+
+SYMS_API void
+syms_arena_free(SymsArena *arena);
+
+#define syms_arena_push_array_ex(arena, type, num, align) (type *)syms_arena_push_ex(arena, sizeof(type)*num, align)
+#define syms_arena_push_struct_ex(arena, type, align) syms_arena_push_array_ex(arena, type, 1, align)
+
+#define syms_arena_push(arena, len)             syms_arena_push_ex(arena, len, sizeof(void*))
+#define syms_arena_push_array(arena, type, num) syms_arena_push_array_ex(arena, type, num, sizeof(void*))
+#define syms_arena_push_struct(arena, type)     syms_arena_push_array(arena, type, 1)
+
+SYMS_API void *
+syms_arena_push_ex(SymsArena *arena, SymsUMM len, SymsUMM align);
+
+SYMS_API char *
+syms_arena_push_cstr(SymsArena *arena, const char *cstr);
+
+SYMS_API SymsArenaFrame
+syms_arena_frame_begin(SymsArena *arena);
+
+SYMS_API void
+syms_arena_frame_end(SymsArenaFrame frame);
+
+// -------------------------------------------------------------------------------- 
+
+SYMS_API U16
+syms_trunc_u16(U32 value);
+
+SYMS_API U32
+syms_trunc_u32(U64 value);
+
+U32
+syms_hash_djb2(const char *str, u32 size);
+
+U16
+syms_hash_djb2_16(const char *str, u32 size);
+
+#endif /* SYMS_CORE_INCLUDE_H */
+
+#if defined(SYMS_CORE_IMPLEMENT)
+
+SYMS_API U32
+syms_addrsize_to_int(SymsAddrSize addr)
+{
+    U32 result;
+    switch (addr) {
+        default:
+        case SYMS_ADDR_SIZE_NULL: result = 0; break;
+        case SYMS_ADDR_SIZE_16:   result = 2; break;
+        case SYMS_ADDR_SIZE_32:   result = 4; break;
+        case SYMS_ADDR_SIZE_64:   result = 8; break;
+        case SYMS_ADDR_SIZE_128:  result = 16; break;
+    }
+    return result;
+}
+
+SYMS_API U16
+syms_trunc_u16(U32 value)
+{
+    U16 result;
+    SYMS_ASSERT_ALWAYS(value <= SYMS_UINT16_MAX);
+    result = (U16)value;
+    return result;
+}
+
+SYMS_API U32
+syms_trunc_u32(U64 value)
+{
+    U32 result;
+    SYMS_ASSERT_ALWAYS(value <= SYMS_UINT32_MAX);
+    result = (U32)value;
+    return result;
+}
+
+SYMS_API void
+syms_arena_init(SymsArena *a, SymsUMM page_size)
+{
+  a->flags = 0;
+  a->size = 0;
+  a->cursor = 0;
+  a->page_size = page_size;
+  a->head = 0;
+}
+
+SYMS_API void
+syms_arena_free(SymsArena *a)
+{
+  void *p = a->head;
+  while (p) {
+    void *n = *(void **)p;
+    *(void **)p = 0;
+    syms_free_virtual_memory(p, 0);
+    p = n;
+  }
+}
+
+SYMS_API void *
+syms_arena_push_ex(SymsArena *a, SymsUMM size, SymsUMM align)
+{
+  SymsUMM cursor = (SymsUMM)a->cursor;
+  SymsUMM mask = align - 1;
+  SymsUMM pad = 0;
+  if (cursor & mask) {
+    pad = align - (cursor & mask);
+  }
+  SymsUMM total_size = size + pad;
+  SymsUMM cursor_offset = cursor - (SymsUMM)a->head;
+  if (cursor_offset + total_size > a->size) {
+    SymsUMM pc = (total_size + 16 + (a->page_size - 1)) / a->page_size;
+    SymsUMM as = pc * a->page_size;
+    u8 *p = (u8 *)syms_reserve_virtual_memory(as);
+    SYMS_ASSERT(p);
+    syms_commit_virtual_memory(p, as);
+    ((void**)p)[0] = a->head;
+    a->head = p;
+    a->cursor = p + 16;
+    a->size = pc * a->page_size;
+  }
+  void *ptr = a->cursor + pad;
+  a->cursor += total_size;
+  return ptr;
+}
+SYMS_API char *
+syms_arena_push_cstr(SymsArena *arena, const char *cstr)
+{
+    u32 cstr_size = (u32)syms_strlen(cstr);
+    char *result = syms_arena_push_array(arena, char, cstr_size);
+    syms_memcpy(result, cstr, cstr_size);
+    return result;
+}
+
+SYMS_API SymsArenaFrame
+syms_arena_frame_begin(SymsArena *a)
+{
+    SymsArenaFrame f;
+    f.arena = a;
+    f.flags = a->flags;
+    f.cursor = a->cursor;
+    f.head = a->head;
+    f.size = a->size;
+    return f;
+}
+
+SYMS_API void
+syms_arena_frame_end(SymsArenaFrame f)
+{
+  SymsArena *a = f.arena;
+  void *p = a->head;
+  while (p != f.head) {
+    void *n = *(void **)p;
+    *(void **)p = 0;
+    syms_free_virtual_memory(p, 0);
+    p = n;
+  }
+  a->head = f.head;
+  a->cursor = f.cursor;
+  a->flags = f.flags;
+  a->size = f.size;
+}
+
+SYMS_API SymsString
+syms_string_init(const char *data, U32 len)
+{
+    SymsString result;
+    while (len > 0 && data[len - 1] == '\0') {
+        --len;
+    }
+    result.data = data;
+    result.len = len;
+    return result;
+}
+
+SYMS_API SymsString
+syms_string_init_lit(const char *data)
+{
+    SymsString result;
+    result.len = 0;
+    result.data = data;
+    while (*data++) {
+        ++result.len;
+    }
+    return result;
+}
+SYMS_API syms_bool
+syms_string_cmp_lit(SymsString a, const char *b)
+{
+    for (;;) {
+        if (a.len == 0 || *b == '\0')
+            break;
+        
+        if (*a.data++ != *b++)
+            break;
+        
+        --a.len;
+    }
+    
+    return a.len == 0 && *b == '\0';
+}
+SYMS_API U32
+syms_string_get_size(SymsString str)
+{
+    SYMS_ASSERT(str.data);
+    return str.len;
+}
+SYMS_API syms_bool
+syms_string_is_null(SymsString str)
+{
+    syms_bool result = str.data == 0 || str.len == 0;
+    return result;
+}
+SYMS_API syms_bool
+syms_string_cmp(SymsString a, SymsString b)
+{
+    syms_bool result = syms_false;
+    if (a.len == b.len) {
+        U32 i;
+        for (i = 0; i < a.len; ++i) {
+            if (a.data[i] != b.data[i])
+                break;
+        }
+        result = i == a.len;
+    }
+    return result;
+}
+SYMS_API char *
+syms_string_to_cstr(SymsString str, SymsArena *arena)
+{
+    char *p = syms_arena_push_array(arena, char, str.len + 1);
+    if (p) {
+        syms_memcpy(p, str.data, str.len);
+        p[str.len] = '\0';
+    }
+    return p;
+}
+
+SYMS_API syms_bool
+syms_is_alpha_ascii(U32 codepoint)
+{
+    return (codepoint >= 'A' && codepoint <= 'Z') || (codepoint >= 'a' && codepoint <= 'z');
+}
+SYMS_API U32
+syms_lowercase(U32 codepoint)
+{
+    if (codepoint >= 'A' && codepoint <= 'Z')
+        codepoint += 'a' - 'A';
+    return codepoint;
+}
+SYMS_API U32
+syms_uppercase(U32 codepoint)
+{
+    if (codepoint >= 'a' && codepoint <= 'z')
+        codepoint -= 'a' - 'A';
+    return codepoint;
+}
+SYMS_API syms_bool
+syms_string_cmp_nocase(SymsString a_str, SymsString b_str)
+{
+    syms_bool result = syms_false;
+    if (a_str.len == b_str.len) {
+        U32 i;
+        for (i = 0; i < a_str.len; ++i) {
+            U32 a = syms_lowercase((U32)a_str.data[i]);
+            U32 b = syms_lowercase((U32)b_str.data[i]);
+            if (a != b) break;
+        }
+        result = (i == a_str.len);
+    }
+    return result;
+}
+SYMS_API char
+syms_string_peek_byte(SymsString str, U32 byte_offset)
+{
+    char result = 0;
+    if (byte_offset < str.len) {
+        result = ((const char *)str.data)[byte_offset];
+    }
+    return result;
+}
+
+SYMS_API const U8 *
+syms_decode_utf8(const U8 *p, U32 *dst)
+{
+    U32 res, n;
+    switch (*p & 0xf0) {
+        case 0xf0 :  res = *p & 0x07;  n = 3;  break;
+        case 0xe0 :  res = *p & 0x0f;  n = 2;  break;
+        case 0xd0 :
+        case 0xc0 :  res = *p & 0x1f;  n = 1;  break;
+        default   :  res = *p;         n = 0;  break;
+    }
+    while (n--) {
+        res = (res << 6) | (*(++p) & 0x3f);
+    }
+    *dst = res;
+    return p + 1;
+}
+
+SYMS_API SymsString
+syms_path_get_file_name(SymsString path, syms_bool strip_ext)
+{
+    const U8 *p = (const U8 *)path.data;
+    const U8 *p_end = p + path.len;
+    const U8 *start = p;
+    const U8 *end = p_end;
+    u32 fname_size;
+    SymsString fname;
+    
+    while (p != p_end) {
+        U32 c;
+        p = syms_decode_utf8(p, &c);
+        if (c == '\0') {
+            break;
+        }
+        if (c == '\\' || c == '/') {
+            start = p;
+        }
+    }
+    if (strip_ext) {
+        p = start;
+        p_end = end;
+        while (p != p_end) {
+            U32 c;
+            p = syms_decode_utf8(p, &c);
+            if (c == '.') {
+                end = p;
+            }
+        }
+    }
+    
+    fname_size = (U32)(end - start);
+    fname = syms_string_init((const char *)start, fname_size);
+    return fname;
+}
+
+SYMS_API syms_bool
+syms_buffer_read(SymsBuffer *bin, void *out_bf, SymsUMM size)
+{
+    syms_bool done = syms_false;
+    if (bin->off + size <= bin->size) {
+        const void *src = (void *)((const char *)bin->base + bin->off);
+        syms_memcpy(out_bf, src, size);
+        bin->off += size;
+        done = syms_true;
+    }
+    return done;
+}
+
+SYMS_API void *
+syms_buffer_push(SymsBuffer *bin, SymsUMM size)
+{
+    void *result = 0;
+    if (bin->off + size <= bin->size) {
+        if (bin->base) {
+            result = (void *)((const U8 *)bin->base + bin->off);
+        }
+        bin->off += size;
+    }
+    return result;
+}
+
+SYMS_API SymsBuffer
+syms_buffer_init_ex(const void *base, SymsUMM size, SymsAddrSize addr_type, SymsAddrSize offs_type)
+{
+    SymsBuffer result;
+    
+    result.base      = base;
+    result.size      = size;
+    result.off       = 0;
+    result.addr_size = syms_addrsize_to_int(addr_type);
+    result.offs_size = syms_addrsize_to_int(offs_type);
+    
+#if defined(SYMS_PARANOID)
+    SYMS_ASSERT(result.addr_size > 0);
+    SYMS_ASSERT(result.offs_size > 0);
+#endif
+    
+    return result;
+}
+
+SYMS_API  SymsBuffer
+syms_buffer_init(const void *data, SymsUMM size)
+{
+    return syms_buffer_init_ex(data, size, SYMS_ADDR_SIZE_HOST, SYMS_ADDR_SIZE_HOST);
+}
+
+SYMS_API syms_bool
+syms_buffer_skip(SymsBuffer *bin, SymsOffset num)
+{
+    syms_bool result = syms_false;
+    if (bin->off + num <= bin->size) {
+        bin->off += num;
+        result = syms_true;
+    }
+    return result;
+}
+
+SYMS_API syms_bool
+syms_buffer_seek(SymsBuffer *bin, SymsOffset new_off)
+{
+    syms_bool result = syms_false;
+    if (new_off <= bin->size) {
+        bin->off = new_off;
+        result = syms_true;
+    }
+    return result;
+}
+
+#define syms_buffer_peek(bin, a) syms_buffer_peek_ex(bin, a, a)
+#define syms_buffer_peek_ex(bin, a, b) \
+do { \
+SymsOffset off = bin->off; \
+a result = syms_buffer_read_##b(bin); \
+bin->off = off; \
+return result; \
+} while (0)
+
+
+SYMS_API U8
+syms_buffer_peek_u8(SymsBuffer *bin)
+{
+    syms_buffer_peek(bin, u8);
+}
+
+SYMS_API U16
+syms_buffer_peek_u16(SymsBuffer *bin)
+{
+    syms_buffer_peek(bin, u16);
+}
+
+SYMS_API U32
+syms_buffer_peek_u32(SymsBuffer *bin)
+{
+    syms_buffer_peek(bin, u32);
+}
+
+SYMS_API U64
+syms_buffer_peek_u64(SymsBuffer *bin)
+{
+    syms_buffer_peek(bin, u64);
+}
+
+SYMS_API S8
+syms_buffer_peek_s8(SymsBuffer *bin)
+{
+    syms_buffer_peek(bin, s8);
+}
+
+SYMS_API S16
+syms_buffer_peek_s16(SymsBuffer *bin)
+{
+    syms_buffer_peek(bin, s16);
+}
+
+SYMS_API S32
+syms_buffer_peek_s32(SymsBuffer *bin)
+{
+    syms_buffer_peek(bin, s32);
+}
+
+SYMS_API S64
+syms_buffer_peek_s64(SymsBuffer *bin)
+{
+    syms_buffer_peek(bin, s64);
+}
+
+SYMS_API S64
+syms_buffer_peek_sleb128(SymsBuffer *bin)
+{
+    syms_buffer_peek_ex(bin, s64, sleb128);
+}
+
+SYMS_API U64
+syms_buffer_peek_uleb128(SymsBuffer *bin)
+{
+    syms_buffer_peek_ex(bin, u64, uleb128);
+}
+
+SYMS_API SymsString
+syms_buffer_peek_string(SymsBuffer *bin)
+{
+    syms_buffer_peek_ex(bin, SymsString, string);
+}
+
+SYMS_API SymsString
+syms_buffer_peek_cstr(SymsBuffer *bin)
+{
+    syms_buffer_peek_ex(bin, SymsString, cstr);
+}
+
+SYMS_API void
+syms_buffer_write(SymsBuffer *bin, const void *src, U32 src_size)
+{
+    void *dst = syms_buffer_push(bin, src_size);
+    if (dst) {
+        syms_memcpy(dst, src, src_size);
+    }
+}
+SYMS_API void
+syms_buffer_write_cstr(SymsBuffer *bin, const char *cstr)
+{
+    U32 cstr_size = (U32)syms_strlen(cstr);
+    syms_buffer_write(bin, cstr, cstr_size);
+}
+SYMS_API const char *
+syms_buffer_get_cstr(SymsBuffer *buffer)
+{
+    char *base = (char *)buffer->base;
+    if (buffer->off >= buffer->size) {
+        SYMS_ASSERT(buffer->size > 0);
+        syms_memset((char*)base + buffer->size - 1, '\0', 1);
+    } else {
+        syms_memset((char*)base + buffer->off, '\0', 1);
+    }
+    return base;
+}
+SYMS_API const void *
+syms_buffer_at(SymsBuffer *buffer)
+{
+    const void *result = (const void *)((const char *)buffer->base + buffer->off);
+    return result;
+}
+
+SYMS_API void
+syms_buffer_write_u8(SymsBuffer *buffer, U8 value)
+{
+    syms_buffer_write(buffer, &value, sizeof(value));
+}
+
+SYMS_API void
+syms_buffer_write_u32(SymsBuffer *buffer, U32 value)
+{
+    syms_buffer_write(buffer, &value, sizeof(value));
+}
+
+SYMS_API SymsString
+syms_buffer_read_cstr(SymsBuffer *buffer)
+{
+    const char *p = (const char *)buffer->base + buffer->off;
+    SymsUWord off = buffer->off;
+    u32 size;
+    SymsString result;
+    while (off < buffer->size) {
+        if (p[off] == '\0') {
+            break;
+        }
+        off += 1;
+    }
+    size = syms_trunc_u32(off - buffer->off);
+    buffer->off = off;
+    result = syms_string_init(p, size);
+    return result;
+}
+
+SYMS_API SymsString
+syms_buffer_read_string(SymsBuffer *buffer)
+{
+    SymsUWord off_start = buffer->off;
+    SymsUWord off_end = buffer->off;
+    const char *ptr = (const char *)buffer->base + buffer->off;
+    SymsString result;
+    while (buffer->off < buffer->size) {
+        if (((const char *)buffer->base)[buffer->off++] == '\0')
+            break;
+    }
+    off_end = buffer->off;
+    result = syms_string_init(ptr, (U32)(off_end - off_start));
+    return result;
+}
+
+SYMS_API U8
+syms_buffer_read_u8(SymsBuffer *buffer)
+{
+    U8 result = 0;
+    syms_buffer_read(buffer, &result, sizeof(result));
+    return result;
+}
+
+SYMS_API U16
+syms_buffer_read_u16(SymsBuffer *buffer)
+{
+    U16 result = 0;
+    syms_buffer_read(buffer, &result, sizeof(result));
+    return result;
+}
+
+SYMS_API U32
+syms_buffer_read_u24(SymsBuffer *buffer)
+{
+    U32 result;
+    U32 a = (U32)syms_buffer_read_u8(buffer);
+    U32 b = (U32)syms_buffer_read_u8(buffer);
+    U32 c = (U32)syms_buffer_read_u8(buffer);
+    result = (a << 0) | (b << 8) | (c << 16);
+    return result;
+}
+
+SYMS_API U32
+syms_buffer_read_u32(SymsBuffer *buffer)
+{
+    U32 result = 0;
+    syms_buffer_read(buffer, &result, sizeof(result));
+    return result;
+}
+
+SYMS_API U64
+syms_buffer_read_u64(SymsBuffer *buffer)
+{
+    U64 result = 0;
+    syms_buffer_read(buffer, &result, sizeof(result));
+    return result;
+}
+
+SYMS_API S8
+syms_buffer_read_s8(SymsBuffer *buffer)
+{
+    S8 result = 0;
+    syms_buffer_read(buffer, &result, sizeof(result));
+    return result;
+}
+
+SYMS_API S16
+syms_buffer_read_s16(SymsBuffer *buffer)
+{
+    S16 result = 0;
+    syms_buffer_read(buffer, &result, sizeof(result));
+    return result;
+}
+
+SYMS_API S32
+syms_buffer_read_s32(SymsBuffer *buffer)
+{
+    S32 result = 0;
+    syms_buffer_read(buffer, &result, sizeof(result));
+    return result;
+}
+
+SYMS_API S64
+syms_buffer_read_s64(SymsBuffer *buffer)
+{
+    S64 result = 0;
+    syms_buffer_read(buffer, &result, sizeof(result));
+    return result;
+}
+
+SYMS_API U64
+syms_buffer_read_uleb128(SymsBuffer *buffer)
+{
+    U64 res = 0;
+    U64 shift = 0;
+    
+    const U8 *data = (const U8 *)buffer->base + buffer->off;
+    const U8 *start = data;
+    
+    U32 num_used;
+    
+    while (buffer->off < buffer->size) {
+        U8 i = *data++;
+        U8 val = i & 0x7fu;
+        res |= (U64)val << shift;
+        if ((i & 0x80u) == 0) {
+            break;
+        }
+        shift += 7u;
+    }
+    
+    num_used = (U32)((const U8 *)data - (const U8 *)start);
+    buffer->off += num_used;
+    
+    return res;
+}
+
+SYMS_API S64
+syms_buffer_read_sleb128(SymsBuffer *buffer)
+{
+    S64 result = 0;
+    U32 shift = 0;
+    
+    const U8 *data = (const U8 *)buffer->base + buffer->off;
+    const U8 *end = (const U8 *)buffer->base + buffer->size;
+    const U8 *start = data;
+    
+    while (data < end) {
+        U8 byte = *data;
+        data += 1;
+        
+        result |= ((U64)(byte & 0x7f)) << shift;
+        shift += 7;
+        
+        if ((byte & 0x80) == 0) {
+            U32 num_used;
+            
+            if (shift < (sizeof(result) * 8) && (byte & 0x40)) {
+                result |= -((S64)(1u << shift));
+            }
+            
+            num_used = (U32)((const U8 *)data - (const U8 *)start);
+            buffer->off += num_used;
+            
+            return (S64)(result);
+        }
+    }
+    
+    return 0;
+}
+
+SYMS_API SymsAddr
+syms_buffer_read_addr(SymsBuffer *buffer)
+{
+    SymsAddr result = 0;
+    SYMS_ASSERT(buffer->addr_size > 0);
+    SYMS_ASSERT(buffer->addr_size <= sizeof(SymsAddr));
+    syms_buffer_read(buffer, &result, buffer->addr_size);
+    return result;
+}
+
+SYMS_API SymsAddr
+syms_buffer_read_offs(SymsBuffer *buffer)
+{
+    SymsAddr result = 0;
+    SYMS_ASSERT(buffer->offs_size > 0);
+    SYMS_ASSERT(buffer->offs_size <= sizeof(SymsAddr));
+    syms_buffer_read(buffer, &result, buffer->offs_size);
+    return result;
+}
+
+U32
+syms_hash_djb2(const char *str, u32 size)
+{
+    U32 hash = 5381;
+    const char *str_end = str + size;
+    while (str < str_end) {
+        U32 c = (U32)*str++;
+        hash = ((hash << 5) + hash) + c; /* hash * 33 + c */
+    }
+    return hash;
+}
+
+U16
+syms_hash_djb2_16(const char *str, U32 size)
+{
+    U32 hash = syms_hash_djb2(str, size);
+    U16 result = (U16)(((hash & 0xFFFF0000) >> 16) + (hash & 0x0000FFFF));
+    return result;
+}
+
+#endif /* SYMS_CORE_IMPLEMENT */
diff --git a/Engine/Source/Runtime/SymsLib/src/syms_enum.c b/Engine/Source/Runtime/SymsLib/src/syms_enum.c
new file mode 100644
--- /dev/null
+++ b/Engine/Source/Runtime/SymsLib/src/syms_enum.c
@@ -0,0 +1,202 @@
+// Copyright Epic Games, Inc. All Rights Reserved.
+SYMS_API const char *
+syms_arch_to_str(SymsArch arch) {
+  const char *str = "";
+  switch (arch) {
+  case SYMS_ARCH_NULL:    str = "NULL";   break;
+  case SYMS_ARCH_ARM:     str = "ARM";    break;
+  case SYMS_ARCH_ARM32:   str = "ARM32";  break;
+  case SYMS_ARCH_PPC:     str = "PPC";    break;
+  case SYMS_ARCH_PPC64:   str = "PPC64";  break;
+  case SYMS_ARCH_IA64:    str = "IA64";   break;
+  case SYMS_ARCH_X64:     str = "X64";    break;
+  case SYMS_ARCH_X86:     str = "X86";    break;
+  }
+  return str;
+}
+
+SYMS_API const char *
+syms_reg_group_to_str(SymsRegClass group)
+{
+  const char *str = "";
+  switch (group) {
+  case SYMS_REG_CLASS_NULL:    str = "NULL";    break;
+  case SYMS_REG_CLASS_STATE:   str = "STATE";   break;
+  case SYMS_REG_CLASS_GPR:     str = "GPR";     break;
+  case SYMS_REG_CLASS_CTRL:    str = "CTRL";    break;
+  case SYMS_REG_CLASS_FP:      str = "FP";      break;
+  case SYMS_REG_CLASS_VEC:     str = "VEC";     break;
+  case SYMS_REG_CLASS_INVALID: str = "INVALID"; break;
+  }
+  return str;
+}
+
+SYMS_API const char *
+syms_member_access_to_str(SymsMemberAccess access)
+{
+  const char *str = "";
+  switch (access) {
+  case SYMS_MEMBER_ACCESS_NULL:      str = "NULL";      break;
+  case SYMS_MEMBER_ACCESS_PRIVATE:   str = "PRIVATE";   break;
+  case SYMS_MEMBER_ACCESS_PUBLIC:    str = "PUBLIC";    break;
+  case SYMS_MEMBER_ACCESS_PROTECTED: str = "PROTECTED"; break;
+  }
+  return str;
+}
+
+SYMS_API const char *
+syms_member_modifier_to_str(SymsMemberModifier modifier)
+{
+  const char *str = "";
+  switch (modifier) {
+  case SYMS_MEMBER_MODIFIER_NULL:         str = "NULL";         break;
+  case SYMS_MEMBER_MODIFIER_VANILLA:      str = "VANILLA";      break;
+  case SYMS_MEMBER_MODIFIER_STATIC:       str = "STATIC";       break;
+  case SYMS_MEMBER_MODIFIER_FRIEND:       str = "FRIEND";       break;
+  case SYMS_MEMBER_MODIFIER_VIRTUAL:      str = "VIRTUAL";      break;
+  case SYMS_MEMBER_MODIFIER_PURE_INTRO:   str = "PURE_INTRO";   break;
+  case SYMS_MEMBER_MODIFIER_INTRO:        str = "INTRO";        break;
+  case SYMS_MEMBER_MODIFIER_PURE_VIRTUAL: str = "PURE_VIRTUAL"; break;
+  }
+  return str;
+}
+
+SYMS_API const char *
+syms_type_modifier_to_str(SymsTypeModifier modifier)
+{
+  const char *str = "";
+  switch (modifier) {
+  case SYMS_TYPE_MDFR_NULL:       str = "NULL";       break;
+  case SYMS_TYPE_MDFR_ATOMIC:     str = "ATOMIC";     break;
+  case SYMS_TYPE_MDFR_CONST:      str = "CONST";      break;
+  case SYMS_TYPE_MDFR_IMMUTABLE:  str = "IMMUTABLE";  break;
+  case SYMS_TYPE_MDFR_REF:        str = "REF";        break;
+  case SYMS_TYPE_MDFR_RESTRICT:   str = "RESTRICT";   break;
+  case SYMS_TYPE_MDFR_RVALUE_REF: str = "RVALUE_REF"; break;
+  case SYMS_TYPE_MDFR_SHARED:     str = "SHARED";     break;
+  case SYMS_TYPE_MDFR_VOLATILE:   str = "VOLATILE";   break;
+  }
+  return str;
+}
+
+SYMS_API const char *
+syms_get_type_kind_info(SymsTypeKind kind)
+{
+  const char *p = "";
+  switch (kind) {
+  case SYMS_TYPE_NULL:          p = "NULL";       break;
+  case SYMS_TYPE_INT8:          p = "INT8";       break;
+  case SYMS_TYPE_INT16:         p = "INT16";      break;
+  case SYMS_TYPE_INT32:         p = "INT32";      break;
+  case SYMS_TYPE_INT64:         p = "INT64";      break;
+  case SYMS_TYPE_INT128:        p = "INT128";     break;
+  case SYMS_TYPE_INT256:        p = "INT256";     break;
+  case SYMS_TYPE_INT512:        p = "INT512";     break;
+  case SYMS_TYPE_UINT8:         p = "UINT8";      break;
+  case SYMS_TYPE_UINT16:        p = "UINT16";     break;
+  case SYMS_TYPE_UINT32:        p = "UINT32";     break;
+  case SYMS_TYPE_UINT64:        p = "UINT64";     break;
+  case SYMS_TYPE_UINT128:       p = "UINT128";    break;
+  case SYMS_TYPE_UINT256:       p = "UINT256";    break;
+  case SYMS_TYPE_UINT512:       p = "UINT512";    break;
+  case SYMS_TYPE_FLOAT16:       p = "FLOAT16";    break;
+  case SYMS_TYPE_FLOAT32:       p = "FLOAT32";    break;
+  case SYMS_TYPE_FLOAT32PP:     p = "FLOAT32PP";  break;
+  case SYMS_TYPE_FLOAT48:       p = "FLOAT48";    break;
+  case SYMS_TYPE_FLOAT64:       p = "FLOAT64";    break;
+  case SYMS_TYPE_FLOAT80:       p = "FLOAT80";    break;
+  case SYMS_TYPE_FLOAT128:      p = "FLOAT128";   break;
+  case SYMS_TYPE_CHAR:          p = "CHAR";       break;
+  case SYMS_TYPE_VOID:          p = "VOID";       break;
+  case SYMS_TYPE_BOOL:          p = "BOOL";       break;
+  case SYMS_TYPE_PTR:           p = "PTR";        break;
+  case SYMS_TYPE_ARR:           p = "ARR";        break;
+  case SYMS_TYPE_ENUM:          p = "ENUM";       break;
+  case SYMS_TYPE_PROC:          p = "PROC";       break;
+  case SYMS_TYPE_PROC_PARAM:    p = "PROC_PARAM"; break;
+  case SYMS_TYPE_TYPEDEF:       p = "TYPEDEF";    break;
+  case SYMS_TYPE_STRUCT:        p = "STRUCT";     break;
+  case SYMS_TYPE_UNION:         p = "UNION";      break;
+  case SYMS_TYPE_CLASS:         p = "CLASS";      break;
+  case SYMS_TYPE_METHOD:        p = "method";     break;
+  case SYMS_TYPE_VIRTUAL_TABLE: p = "VIRTUAL_TABLE"; break;
+  case SYMS_TYPE_BASE_CLASS:    p = "BASE CLASS";  break;
+  case SYMS_TYPE_BITFIELD:      p = "BITFIELD";    break;
+  case SYMS_TYPE_COMPLEX32:     p = "COMPLEX32";   break;
+  case SYMS_TYPE_COMPLEX64:     p = "COMPLEX64";   break;
+  case SYMS_TYPE_COMPLEX80:     p = "COMPLEX80";   break;
+  case SYMS_TYPE_COMPLEX128:    p = "COMPLEX128";  break;
+  case SYMS_TYPE_VARIADIC:      p = "VARIADIC";    break;
+  case SYMS_TYPE_STRING:        p = "STRING";      break;
+  case SYMS_TYPE_WCHAR:         p = "WCHAR";       break;
+  }
+  return p;
+}
+
+SYMS_API syms_uint
+syms_get_type_kind_bitcount(SymsArch arch, SymsTypeKind kind)
+{
+  syms_uint result = 0;
+  switch (kind) {
+  case SYMS_TYPE_INT8:          result = 8;   break;
+  case SYMS_TYPE_INT16:         result = 16;  break;
+  case SYMS_TYPE_INT32:         result = 32;  break;
+  case SYMS_TYPE_INT64:         result = 64;  break;
+  case SYMS_TYPE_UINT8:         result = 8;   break;
+  case SYMS_TYPE_UINT16:        result = 16;  break;
+  case SYMS_TYPE_UINT32:        result = 32;  break;
+  case SYMS_TYPE_UINT64:        result = 64;  break;
+  case SYMS_TYPE_INT128:        result = 128; break;
+  case SYMS_TYPE_INT256:        result = 256; break;
+  case SYMS_TYPE_INT512:        result = 512; break;
+  case SYMS_TYPE_UINT128:       result = 128; break;
+  case SYMS_TYPE_UINT256:       result = 256; break;
+  case SYMS_TYPE_UINT512:       result = 512; break;
+  case SYMS_TYPE_FLOAT16:       result = 16;  break;
+  case SYMS_TYPE_FLOAT32:       result = 32;  break;
+  case SYMS_TYPE_FLOAT32PP:     result = 32;  break;
+  case SYMS_TYPE_FLOAT48:       result = 48;  break;
+  case SYMS_TYPE_FLOAT64:       result = 64;  break;
+  case SYMS_TYPE_FLOAT80:       result = 80;  break;
+  case SYMS_TYPE_FLOAT128:      result = 128; break;
+  case SYMS_TYPE_CHAR:          result = 8;   break;
+  case SYMS_TYPE_VOID:          result = syms_get_addr_size_ex(arch); break; 
+  case SYMS_TYPE_TYPEDEF:       result = 0;   break;
+  case SYMS_TYPE_STRUCT:        result = 0;   break;
+  case SYMS_TYPE_UNION:         result = 0;   break;
+  case SYMS_TYPE_ENUM:          result = 8;   break;
+  case SYMS_TYPE_PROC:          result = syms_get_addr_size_ex(arch); break;
+  case SYMS_TYPE_PROC_PARAM:    result = 0;   break;
+  case SYMS_TYPE_ARR:           result = 0;   break;
+  case SYMS_TYPE_PTR:           result = 0;   break;
+  case SYMS_TYPE_BITFIELD:      result = 0;   break;
+  case SYMS_TYPE_COMPLEX32:     result = 32;  break;
+  case SYMS_TYPE_COMPLEX64:     result = 64;  break;
+  case SYMS_TYPE_COMPLEX128:    result = 128; break;
+  case SYMS_TYPE_VARIADIC:      result = 0;   break;
+  case SYMS_TYPE_CLASS:         result = 0;   break;
+  case SYMS_TYPE_METHOD:        result = 0;   break;
+  case SYMS_TYPE_VIRTUAL_TABLE: result = 0;   break;
+  case SYMS_TYPE_BASE_CLASS:    result = 0;   break;
+  case SYMS_TYPE_BOOL:          result = 8;   break;
+  case SYMS_TYPE_INVALID:       result = 0;   break;
+  }
+  SYMS_ASSERT(result % 8 == 0);
+  return result;
+}
+
+SYMS_API syms_uint
+syms_get_addr_size_ex(SymsArch arch)
+{
+  syms_uint addr_size = 0;
+
+  switch (arch) {
+  case SYMS_ARCH_NULL:  addr_size = 0; break;
+  case SYMS_ARCH_X86:   addr_size = 4; break;
+  case SYMS_ARCH_X64:   addr_size = 8; break;
+  default: SYMS_INVALID_CODE_PATH;
+  }
+  
+  return addr_size;
+}
+
diff --git a/Engine/Source/Runtime/SymsLib/src/syms_line_table.c b/Engine/Source/Runtime/SymsLib/src/syms_line_table.c
new file mode 100644
--- /dev/null
+++ b/Engine/Source/Runtime/SymsLib/src/syms_line_table.c
@@ -0,0 +1,249 @@
+// Copyright Epic Games, Inc. All Rights Reserved.
+/******************************************************************************
+ * File   : syms_line_table.c                                                 *
+ * Author : Nikita Smith                                                      *
+ * Created: 2020/08/07                                                        *
+ * Purpose: Contains routines for mapping address to source line and          *
+ *          vice versa                                                        *
+ ******************************************************************************/
+
+SYMS_INTERNAL int
+syms_qsort_addr_map_compar(const void *a, const void *b)
+{
+  const SymsAddrMap *map_a = (const SymsAddrMap *)a;
+  const SymsAddrMap *map_b = (const SymsAddrMap *)b;
+  if (map_a->addr < map_b->addr) return -1;
+  if (map_a->addr > map_b->addr) return +1;
+  // SYMS_ASSERT_FAILURE_PARANOID("address collision");
+  return 0;
+}
+
+SYMS_INTERNAL SymsFileID
+syms_line_table_push_file_(SymsLineTable *lt, SymsInstance *instance, SymsStringRef *name_ref, syms_bool register_file)
+{
+  SymsFileID result = SYMS_INVALID_FILE_ID;
+
+  char temp[1024];
+  char temp1[1024];
+
+  if (lt->file_count > 0) {
+	  syms_uint temp_size = syms_read_strref(instance, name_ref, temp, sizeof(temp));
+	  SymsFileID file_hash = syms_hash_djb2(temp, temp_size) % SYMS_INVALID_FILE_ID;
+	  SymsFileID best_id = file_hash % lt->file_count;
+	  SymsFileID id = best_id;
+
+
+	  do {
+		  if (lt->files[id].type == SYMS_STRING_REF_TYPE_NULL) {
+			  if (register_file) {
+				  lt->files[id] = *name_ref;
+			  }
+			  result = id;
+			  break;
+		  }
+		  syms_read_strref(instance, &lt->files[id], temp1, sizeof(temp1));
+		  if (syms_strcmp(temp1, temp) == 0) {
+			  result = id;
+			  break;
+		  }
+		  id = (id + 1) % lt->file_count;
+	  } while (id != best_id);
+  }
+  return result;
+}
+
+SYMS_INTERNAL SymsFileID
+syms_line_table_push_file(SymsLineTable *lt, SymsInstance *instance, SymsStringRef *name_ref)
+{
+  return syms_line_table_push_file_(lt, instance, name_ref, syms_true);
+}
+
+SYMS_INTERNAL SymsFileID
+syms_line_table_find_file(SymsLineTable *lt, SymsInstance *instance, SymsStringRef *name_ref)
+{
+  return syms_line_table_push_file_(lt, instance, name_ref, syms_false);
+}
+
+SYMS_INTERNAL SymsErrorCode
+syms_line_table_build(SymsLineTable *lt, SymsInstance *instance, SymsMod *mod, SymsArena *arena)
+{
+  SymsLineIter line_iter;
+  if (syms_line_iter_init(instance, &line_iter, mod)) {
+    SymsLine line;
+    syms_uint line_count = 0;
+    syms_uint file_count = 0;
+    SymsUMM block_size;
+    u8 *block, *ptr;
+    SymsFileID curr_file_id;
+
+    while (syms_line_iter_next(&line_iter, &line)) {
+      if (line_iter.switched_file) {
+        ++file_count;
+      }
+      ++line_count;
+    }
+
+    block_size = line_count * (sizeof(lt->addrs[0]) + sizeof(lt->lines[0])) +
+                 file_count * sizeof(lt->files[0]);
+    if (block_size == 0) {
+      return SYMS_ERR_OK; // no line table
+    }
+    block = syms_arena_push_array(arena, u8, block_size);
+    if (block == NULL) {
+      return SYMS_ERR_NOMEM;
+    }
+    ptr = block;
+    lt->line_count = 0;
+    lt->line_max = line_count;
+    lt->file_count = file_count;
+    lt->addrs = (SymsAddrMap *)ptr;
+    ptr = (ptr + sizeof(lt->addrs[0]) * line_count);
+    lt->lines = (SymsSourceMap *)ptr;
+    ptr = (ptr + sizeof(lt->lines[0]) * line_count);
+    lt->files = (SymsStringRef *)ptr;
+    ptr = (ptr + sizeof(lt->files[0]) * file_count);
+    // clear "files" array to zero, a zeroed entry is considered an empty
+    // bucket
+    syms_memset(lt->files, 0, file_count * sizeof(lt->files[0]));
+    // make sure we exhausted allocated block
+    SYMS_ASSERT(ptr == (void *)(block + block_size));
+    curr_file_id = SYMS_INVALID_FILE_ID;
+    syms_line_iter_init(instance, &line_iter, mod);
+    while (syms_line_iter_next(&line_iter, &line)) {
+      if (line_iter.switched_file) {
+        curr_file_id = syms_line_table_push_file(lt, instance, &line_iter.file.name);
+      }
+      SYMS_ASSERT(lt->line_count < lt->line_max);
+      lt->addrs[lt->line_count].addr = line.va;
+      lt->addrs[lt->line_count].id = lt->line_count;
+      lt->lines[lt->line_count].ln = line.ln;
+      lt->lines[lt->line_count].col = line.col;
+      lt->lines[lt->line_count].file = curr_file_id;
+      lt->line_count += 1;
+    }
+
+    syms_qsort(lt->addrs, lt->line_count, sizeof(lt->addrs[0]), syms_qsort_addr_map_compar);
+  }
+
+  return SYMS_ERR_OK;
+}
+
+SYMS_API syms_bool
+syms_line_table_map_va(SymsLineTable *lt, SymsAddr va, SymsSourceFileMap *map_out)
+{
+  SymsLineTableQuery query;
+  if (syms_line_table_map_va_ex(lt, va, &query)) {
+    *map_out = query.map;
+    return syms_true;
+  }
+  return syms_false;
+}
+
+SYMS_API syms_bool
+syms_line_table_map_va_ex(SymsLineTable *lt, SymsAddr va, SymsLineTableQuery *query_out)
+{
+  syms_int nearest_index = -1;
+  SymsAddr nearest_delta = SYMS_ADDR_MAX;
+  syms_bool is_va_mapped = syms_false;
+
+  if (lt) {
+    if (lt->line_count > 1) {
+      if (va >= lt->addrs[0].addr && va <= lt->addrs[lt->line_count - 1].addr) {
+        syms_int max = (S32)lt->line_count - 1;
+        syms_int min = 0;
+        syms_int mid = -1;
+        while (min <= max) {
+          mid = min + (max - min + 1)/2;
+          if (va < lt->addrs[mid].addr) {
+            max = mid - 1;
+          } else if (va > lt->addrs[mid].addr) {
+            SymsAddr delta = va - lt->addrs[mid].addr;
+            if (delta < nearest_delta) {
+              nearest_index = mid;
+              nearest_delta = delta;
+            }
+            min = mid + 1;
+          } else {
+            nearest_index = mid;
+            nearest_delta = 0;
+            break;
+          }
+        }
+      }
+    } else if (lt->line_count == 1) {
+      if (va >= lt->addrs[0].addr) {
+        nearest_index = 0;
+        nearest_delta = 0;
+      }
+    }
+
+    if (nearest_index >= 0 && lt->line_count > 0) {
+      syms_uint line_id;
+      SymsSourceMap *src;
+
+      SYMS_ASSERT((syms_uint)nearest_index < lt->line_count);
+
+      line_id = lt->addrs[nearest_index].id; 
+      src = &lt->lines[line_id];
+
+      query_out->map.file.name        = lt->files[src->file];
+      query_out->map.file.chksum_type = SYMS_CHECKSUM_NULL;
+      query_out->map.line.va          = lt->addrs[nearest_index].addr;
+      query_out->map.line.ln          = src->ln;
+      query_out->map.line.col         = src->col;
+      
+      if ((syms_uint)nearest_index + 1 < lt->line_count) {
+        SymsAddr a = lt->addrs[nearest_index + 0].addr;
+        SymsAddr b = lt->addrs[nearest_index + 1].addr;
+        query_out->map.instructions_size = (syms_uint)(b - a);
+      } else {
+        query_out->map.instructions_size = 0;
+      }
+
+      is_va_mapped = syms_true;
+    }
+  }
+
+  return is_va_mapped;
+}
+
+SYMS_API syms_bool
+syms_line_table_map_src(SymsLineTable *lt, SymsInstance *instance, SymsString filename, syms_uint ln, SymsSourceFileMap *map_out)
+{
+  SymsStringRef file_ref = syms_string_ref_str(filename);
+  SymsFileID file = syms_line_table_find_file(lt, instance, &file_ref);
+  syms_bool is_src_mapped = syms_false;
+  if (file != SYMS_INVALID_FILE_ID) {
+    syms_uint nearest_index = SYMS_UINT_MAX;
+    syms_uint nearest_delta = SYMS_UINT_MAX;
+    syms_uint i;
+    for (i = 0; i < lt->line_count; ++i) {
+      if (lt->lines[i].file == file) {
+        syms_uint delta = (syms_uint)SYMS_ABS((syms_int)lt->lines[i].ln - (syms_int)ln);
+        if (delta < nearest_delta) {
+          nearest_index = i;
+          nearest_delta = delta;
+        }
+      }
+    }
+    for (i = 0; i < lt->line_count; ++i) {
+      if (lt->addrs[i].id == nearest_index) {
+        map_out->file.name    = lt->files[file];
+        map_out->file.chksum_type = SYMS_CHECKSUM_NULL;
+        map_out->line.va      = lt->addrs[i].addr;
+        map_out->line.ln      = lt->lines[nearest_index].ln;
+        map_out->line.col     = lt->lines[nearest_index].col;
+
+        if (i + 1 < lt->line_count) {
+          map_out->instructions_size = (syms_uint)(lt->addrs[i + 1].addr - lt->addrs[i].addr);
+        } else {
+          map_out->instructions_size = 0;
+        }
+
+        is_src_mapped = syms_true;
+        break;
+      }
+    }
+  }
+  return is_src_mapped;
+}
diff --git a/Engine/Source/Runtime/SymsLib/src/syms_line_table.h b/Engine/Source/Runtime/SymsLib/src/syms_line_table.h
new file mode 100644
--- /dev/null
+++ b/Engine/Source/Runtime/SymsLib/src/syms_line_table.h
@@ -0,0 +1,35 @@
+// Copyright Epic Games, Inc. All Rights Reserved.
+#ifndef SYMS_LINE_TABLE_INCLUDE_H
+#define SYMS_LINE_TABLE_INCLUDE_H
+
+/******************************************************************************
+ * File   : syms_line_table.h                                                 *
+ * Author : Nikita Smith                                                      *
+ * Created: 2020/08/07                                                        *
+ ******************************************************************************/
+
+#if 0
+typedef struct SymsLineInfo
+{
+  SymsAddr addr;
+  U32 ln;
+  U16 col;
+  U16 file;
+} SymsLineInfo;
+#endif
+
+#if 0
+typedef struct SymsLinesSec 
+{
+  U32 ln_min;
+  U32 ln_max;
+  U32 lines_count;
+  SymsAddr addr_min;
+  SymsAddr addr_max;
+  struct SymsLine *lines;
+  struct SymsLinesSec *next;
+  struct SymsSourceFile file;
+} SymsLinesSec;
+#endif
+
+#endif // SYMS_LINE_TABLE_INCLUDE_H
diff --git a/Engine/Source/Runtime/SymsLib/src/syms_macos.c b/Engine/Source/Runtime/SymsLib/src/syms_macos.c
new file mode 100644
--- /dev/null
+++ b/Engine/Source/Runtime/SymsLib/src/syms_macos.c
@@ -0,0 +1,42 @@
+// Copyright Epic Games, Inc. All Rights Reserved.
+#include <sys/mman.h>
+#include <unistd.h>
+#include <errno.h>
+
+SYMS_API void
+syms_init_os()
+{
+}
+SYMS_API SymsUMM
+syms_get_pagesize(void)
+{
+  int size = getpagesize();
+  return (SymsUMM)size;
+}
+SYMS_API void *
+syms_reserve_virtual_memory(SymsUMM size)
+{
+  void *base = mmap(0, size, PROT_NONE, MAP_PRIVATE|MAP_ANONYMOUS, -1, 0);
+  return base;
+}
+SYMS_API syms_bool
+syms_commit_virtual_memory(void *base, SymsUMM size)
+{
+  int result = mprotect(base, size, PROT_READ|PROT_WRITE);
+  if (result != 0) {
+    switch (errno) {
+    case EACCES:  SYMS_ASSERT_FAILURE("unable to change page flags"); break;
+    case EINVAL:  SYMS_ASSERT_FAILURE("addres is not page-aligned"); break;
+    case ENOTSUP: SYMS_ASSERT_FAILURE("invalid combination of protection flags"); break;
+    default:      SYMS_ASSERT_FAILURE("mprotect");
+    }
+  }
+  return (result == 0);
+}
+SYMS_API void
+syms_free_virtual_memory(void *base, SymsUMM size)
+{
+  if (munmap(base, size) != 0)
+    SYMS_ASSERT_FAILURE("unable to unmap memory");
+}
+
diff --git a/Engine/Source/Runtime/SymsLib/src/syms_platform.h b/Engine/Source/Runtime/SymsLib/src/syms_platform.h
new file mode 100644
--- /dev/null
+++ b/Engine/Source/Runtime/SymsLib/src/syms_platform.h
@@ -0,0 +1,44 @@
+// Copyright Epic Games, Inc. All Rights Reserved.
+#if !defined(SYMS_PLATFORM_INCLUDE_H)
+#define SYMS_PLATFORM_INCLUDE_H
+
+SYMS_API void
+syms_init_os(void);
+
+SYMS_API SymsUMM
+syms_get_pagesize(void);
+
+SYMS_API void *
+syms_reserve_virtual_memory(SymsUMM size);
+
+SYMS_API syms_bool
+syms_commit_virtual_memory(void *base, SymsUMM size);
+
+SYMS_API void
+syms_free_virtual_memory(void *base, SymsUMM size);
+
+SYMS_INLINE void *
+syms_virtual_alloc(SymsUMM size)
+{
+  void *base = syms_reserve_virtual_memory(size); 
+  if (!syms_commit_virtual_memory(base, size)) {
+    SYMS_ASSERT_FAILURE_PARANOID("unable to commit virtual memory");
+    base = 0;
+  }
+  return base;
+}
+
+typedef struct
+{
+  void *base; // pointer to first byte in the file
+  SymsUMM size; 
+} SymsEntireFile;
+
+SYMS_API SymsEntireFile 
+syms_read_entire_file(const char *file_name);
+
+SYMS_API void
+syms_free_entire_file(SymsEntireFile *file);
+
+#endif // SYMS_PLATFORM_INCLUDE_H
+
