RAD symbol library

diff --git a/Engine/Source/Runtime/SymsLib/src/win/syms_nt_unwind.c b/Engine/Source/Runtime/SymsLib/src/win/syms_nt_unwind.c
new file mode 100644
--- /dev/null
+++ b/Engine/Source/Runtime/SymsLib/src/win/syms_nt_unwind.c
@@ -0,0 +1,716 @@
+// Copyright Epic Games, Inc. All Rights Reserved.
+/******************************************************************************
+ * File   : syms_nt_unwind.c                                                  *
+ * Author : Nikita Smith                                                      *
+ * Created: 2020/09/11                                                        *
+ * Purpose: Stack unwinder for Win32                                          *
+ ******************************************************************************/
+
+SYMS_INTERNAL SymsErrorCode
+syms_memread_pdata(struct SymsInstance *instance, SymsMemread *memread_info, SymsAddr va, SymsNTPdata *pdata_out)
+{
+  SymsNTPdataPacked packed_pdata;
+  SymsErrorCode result = syms_memread(memread_info, va, &packed_pdata, sizeof(packed_pdata));
+  if (SYMS_RESULT_OK(result)) {
+    *pdata_out = syms_unpack_pdata(instance, &packed_pdata);
+  }
+  return result;
+}
+
+SYMS_INTERNAL syms_uint
+syms_nt_unwind_info_sizeof(SymsNTUnwindInfo *uwinfo)
+{
+  syms_uint size = sizeof(SymsNTUnwindInfo);
+  size += ((uwinfo->codes_num + 1) & ~1) * sizeof(SymsNTUnwindInfo);
+  if (SYMS_NT_UNWIND_INFO_HEADER_GET_FLAGS(uwinfo->header) & SYMS_NT_UNWIND_INFO_CHAINED) {
+    size += sizeof(SymsNTPdataPacked);
+  } else if (SYMS_NT_UNWIND_INFO_HEADER_GET_FLAGS(uwinfo->header) & SYMS_NT_UNWIND_INFO_FHANDLER) {
+    size += sizeof(syms_uint);
+    size += sizeof(SymsNTPdataPacked);
+  }
+  return size;
+}
+
+SYMS_INTERNAL syms_uint
+syms_nt_unwind_code_count_nodes(U8 uwcode_flags)
+{
+  static syms_uint counts[] = {
+#define X(name, value, numcodes) numcodes,
+    SYMS_NT_UWOP_LIST
+#undef X
+  };
+  
+  syms_uint op = SYMS_NT_UNWIND_CODE_FLAGS_GET_OP_CODE(uwcode_flags);
+  syms_uint info = SYMS_NT_UNWIND_CODE_FLAGS_GET_OP_INFO(uwcode_flags);
+  syms_uint result = 0;
+  if (op < SYMS_ARRAY_SIZE(counts)) {
+    result = counts[op];
+    if (op == SYMS_NT_OP_ALLOC_LARGE && info == 1) {
+      result += 1;
+    }
+  }
+  return result;
+}
+
+SYMS_INTERNAL SymsRegID
+syms_remap_gpr_nt(SymsImage *img, syms_uint nt_regid)
+{
+  SymsRegID result;
+  switch (img->arch) {
+  case SYMS_ARCH_X86: {
+    switch (nt_regid) {
+    case 0 : result = SYMS_REG_X86_eax; break;
+    case 1 : result = SYMS_REG_X86_ecx; break;
+    case 2 : result = SYMS_REG_X86_edx; break;
+    case 3 : result = SYMS_REG_X86_ebx; break;
+    case 4 : result = SYMS_REG_X86_esp; break;
+    case 5 : result = SYMS_REG_X86_ebp; break;
+    case 6 : result = SYMS_REG_X86_esi; break;
+    case 7 : result = SYMS_REG_X86_edi; break;
+    default: result = SYMS_REG_null; break;
+    }
+  } break;
+  case SYMS_ARCH_X64: {
+    switch (nt_regid) {
+    case 0 : result = SYMS_REG_X64_rax; break;
+    case 1 : result = SYMS_REG_X64_rcx; break;
+    case 2 : result = SYMS_REG_X64_rdx; break;
+    case 3 : result = SYMS_REG_X64_rbx; break;
+    case 4 : result = SYMS_REG_X64_rsp; break;
+    case 5 : result = SYMS_REG_X64_rbp; break;
+    case 6 : result = SYMS_REG_X64_rsi; break;
+    case 7 : result = SYMS_REG_X64_rdi; break;
+    case 8 : result = SYMS_REG_X64_r8;  break;
+    case 9 : result = SYMS_REG_X64_r9;  break;
+    case 10: result = SYMS_REG_X64_r10; break;
+    case 11: result = SYMS_REG_X64_r11; break;
+    case 12: result = SYMS_REG_X64_r12; break;
+    case 13: result = SYMS_REG_X64_r13; break;
+    case 14: result = SYMS_REG_X64_r14; break;
+    case 15: result = SYMS_REG_X64_r15; break;
+    default: result = SYMS_REG_null; break;
+    }
+  } break;
+  default: result = SYMS_REG_null; break;
+  }
+  SYMS_ASSERT(result != SYMS_REG_null);
+  return result;
+}
+
+SYMS_INTERNAL SymsRegID
+syms_remap_xmm_nt(SymsImage *img, syms_uint nt_regid)
+{
+  SymsRegID result;
+  switch (img->arch) {
+  case SYMS_ARCH_X86: {
+    switch (nt_regid) {
+    case 0 : result = SYMS_REG_X86_xmm0;  break;
+    case 1 : result = SYMS_REG_X86_xmm1;  break;
+    case 2 : result = SYMS_REG_X86_xmm2;  break;
+    case 3 : result = SYMS_REG_X86_xmm3;  break;
+    case 4 : result = SYMS_REG_X86_xmm4;  break;
+    case 5 : result = SYMS_REG_X86_xmm5;  break;
+    case 6 : result = SYMS_REG_X86_xmm6;  break;
+    case 7 : result = SYMS_REG_X86_xmm7;  break;
+    default: result = SYMS_REG_null;   break;
+    }
+  } break;
+  case SYMS_ARCH_X64: {
+    switch (nt_regid) {
+    case 0 : result = SYMS_REG_X64_xmm0;  break;
+    case 1 : result = SYMS_REG_X64_xmm1;  break;
+    case 2 : result = SYMS_REG_X64_xmm2;  break;
+    case 3 : result = SYMS_REG_X64_xmm3;  break;
+    case 4 : result = SYMS_REG_X64_xmm4;  break;
+    case 5 : result = SYMS_REG_X64_xmm5;  break;
+    case 6 : result = SYMS_REG_X64_xmm6;  break;
+    case 7 : result = SYMS_REG_X64_xmm7;  break;
+    case 8 : result = SYMS_REG_X64_xmm8;  break;
+    case 9 : result = SYMS_REG_X64_xmm9;  break;
+    case 10: result = SYMS_REG_X64_xmm10; break;
+    case 11: result = SYMS_REG_X64_xmm11; break;
+    case 12: result = SYMS_REG_X64_xmm12; break;
+    case 13: result = SYMS_REG_X64_xmm13; break;
+    case 14: result = SYMS_REG_X64_xmm14; break;
+    case 15: result = SYMS_REG_X64_xmm15; break;
+    default: result = SYMS_REG_null;   break;
+    }
+  } break;
+  default: result = SYMS_REG_null; break;
+  }
+  SYMS_ASSERT(result != SYMS_REG_null);
+  return result;
+}
+
+
+SYMS_INTERNAL syms_bool
+syms_is_ip_inside_epilog(SymsBuffer inst_cursor, SymsNTPdata pdata, SymsAddr ip)
+{
+  syms_bool is_epilog = syms_false;
+  U8 inst_bytes[4];
+
+  if (syms_buffer_read(&inst_cursor, &inst_bytes[0], sizeof(inst_bytes))) {
+    syms_bool keep_parsing = syms_true;
+    syms_buffer_seek(&inst_cursor, 0);
+    if ((inst_bytes[0] & 0xF8) == 0x48) {
+      switch (inst_bytes[1]) {
+      case 0x81: { // add $nnnn,%rsp
+        if (inst_bytes[0] == 0x48 && inst_bytes[2] == 0xC4) {
+          syms_buffer_skip(&inst_cursor, 7);
+        } else {
+          keep_parsing = syms_false;
+        }
+      } break;
+      case 0x83: { // add $n,%rsp
+        if (inst_bytes[0] == 0x48 && inst_bytes[2] == 0xC4) {
+          syms_buffer_skip(&inst_cursor, 4);
+        } else {
+          keep_parsing = syms_false;
+        }
+      } break;
+      case 0x8D: { // lea n(reg),%rsp
+        keep_parsing = syms_false;
+        if (inst_bytes[0] & 0x06) {
+          break; // rex.RX must be cleared 
+        }
+        if (((inst_bytes[2] >> 3) & 0x07) != 0x04) {
+          break; // dest reg mus be %rsp
+        }
+        if ((inst_bytes[2] & 0x07) == 0x04) {
+          break; // no SIB byte allowed
+        }
+        if ((inst_bytes[2] >> 6) == 1) { // 8-bit offset
+          syms_buffer_skip(&inst_cursor, 4);
+          keep_parsing = syms_true;
+          break;
+        }
+        if ((inst_bytes[2] >> 6) == 0x02) { // 32-bit offset
+          syms_buffer_skip(&inst_cursor, 7);
+          keep_parsing = syms_true;
+          break;
+        }
+      } break;
+      }
+    }
+
+    for (; keep_parsing; ) {
+      U8 inst_byte;
+
+      keep_parsing = syms_false;
+      inst_byte = syms_buffer_read_u8(&inst_cursor);
+      if ((inst_byte & 0xF0) == 0x40) {
+        inst_byte = syms_buffer_read_u8(&inst_cursor);
+      }
+      switch (inst_byte) {
+      case 0x58:   // pop %rax/%r8 
+      case 0x59:   // pop %rcx/%r9 
+      case 0x5A:   // pop %rdx/%r10
+      case 0x5B:   // pop %rbx/%r11
+      case 0x5C:   // pop %rsp/%r12
+      case 0x5D:   // pop %rbp/%r13
+      case 0x5E:   // pop %rsi/%r14
+      case 0x5F: { // pop %rdi/%r15
+        keep_parsing = syms_true;
+      } break;
+      case 0xC2:   // ret $nn
+      case 0xC3: { // ret 
+        is_epilog = syms_true;
+      } break;
+      case 0xE9: { // jmp nnnn
+        S32 imm = syms_buffer_read_s32(&inst_cursor);
+        S64 off = (S64)inst_cursor.off + imm;
+        if (off < 0) {
+          return syms_false;
+        }
+        syms_buffer_seek(&inst_cursor, (SymsOffset)off);
+        keep_parsing = ip + inst_cursor.off >= pdata.lo && ip + inst_cursor.off < pdata.hi;
+      } break;
+      case 0xF3: { // rep; ret (for amd64 prediction bug) 
+        SymsUWord inst_cursor_off = inst_cursor.off;
+        U8 next_inst_byte = syms_buffer_read_u8(&inst_cursor);
+        is_epilog = next_inst_byte == 0xc3;
+        syms_buffer_seek(&inst_cursor, inst_cursor_off);
+      } break;
+      }
+    }
+  }
+
+  return is_epilog;
+}
+
+SYMS_INTERNAL SymsErrorCode
+syms_unwind_from_epilog(SymsImage *img,
+    SymsArch arch,
+    SymsMemread *memread_info, 
+    void *regread_context, syms_regread_sig *regread, 
+    void *regwrite_context, syms_regwrite_sig *regwrite, 
+    SymsAddr ip, SymsAddr sp, 
+    SymsBuffer inst_cursor)
+{
+  syms_uint gpr_size = (syms_uint)syms_get_addr_size_ex(arch);
+  SymsRegID sp_regid = syms_regs_get_sp_ex(arch);
+  SymsRegID ip_regid = syms_regs_get_ip_ex(arch);
+
+  SymsErrorCode result = SYMS_ERR_INVALID_CODE_PATH;
+  syms_bool keep_parsing = syms_true;
+  while (keep_parsing && inst_cursor.off < inst_cursor.size) {
+    U8 inst_byte;
+    U8 rex;
+
+    keep_parsing = syms_false;
+
+    inst_byte = syms_buffer_read_u8(&inst_cursor);
+    rex = 0;
+    if ((inst_byte & 0xF0) == 0x40) {
+      rex = inst_byte & 0x0F; // rex prefix
+      inst_byte = syms_buffer_read_u8(&inst_cursor);
+    }
+
+    switch (inst_byte) {
+    case 0x58:   // pop %rax/r8
+    case 0x59:   // pop %rcx/r9
+    case 0x5A:   // pop %rdx/r10
+    case 0x5B:   // pop %rbx/r11
+    case 0x5C:   // pop %rsp/r12
+    case 0x5D:   // pop %rbp/r13
+    case 0x5E:   // pop %rsi/r14
+    case 0x5F: { // pop %rdi/r15
+      U8 reg = inst_byte - 0x58 + (rex & 1) * 8;
+      SymsRegID regid = syms_remap_gpr_nt(img, reg);
+      syms_uint read_size = regwrite(regwrite_context, arch, regid, &sp, gpr_size);
+      keep_parsing = (read_size == gpr_size);
+      SYMS_ASSERT_PARANOID(keep_parsing);
+      sp += 8;
+    } break;
+
+    case 0x81: { // add $nnnn,%rsp 
+      S32 imm;
+      syms_buffer_read_u8(&inst_cursor);
+      imm = syms_buffer_read_s32(&inst_cursor);
+      if (imm < 0 && (S64)imm < (S64)sp) {
+        break;
+      }
+      sp = (SymsAddr)((S64)sp + (S64)imm);
+      keep_parsing = syms_true;
+    } break;
+
+    case 0x83: { // add $n,%rsp
+      S8 imm;
+      syms_buffer_read_u8(&inst_cursor);
+      imm = syms_buffer_read_s8(&inst_cursor);
+      if (imm < 0 && (S64)imm < (S64)sp) {
+        return syms_false;
+      }
+      sp = (SymsAddr)((S64)sp + (S64)imm);
+      keep_parsing = syms_true;
+    } break;
+
+    case 0x8D: { // lea imm8/imm32,$rsp
+      U8 modrm = syms_buffer_read_u8(&inst_cursor);
+      U8 reg = (modrm & 7) + (rex & 1)*8;
+      SymsRegID regid = syms_remap_gpr_nt(img, reg);
+      U64 reg_value = 0;
+      regread(regread_context, arch, regid, &reg_value, sizeof(reg_value));
+      S32 imm;
+      if ((modrm >> 6) == 1) { // lea n(reg),%rsp
+        imm = (S32)syms_buffer_read_s8(&inst_cursor);
+      } else { // lea nnnn(reg),%rsp
+        imm = syms_buffer_read_s32(&inst_cursor);
+      }
+      if (imm < 0 && imm < (S64)reg_value) {
+        break;
+      }
+      sp = (SymsAddr)((S64)reg_value + imm);
+      keep_parsing = syms_true;
+    } break;
+
+    case 0xC2: { // ret $nn
+      u16 imm;
+      result = syms_memread(memread_info, sp, &ip, gpr_size);
+      if (SYMS_RESULT_FAIL(result)) {
+        return result;
+      }
+      imm = syms_buffer_read_u16(&inst_cursor);
+      sp += gpr_size + imm;
+    } break;
+    case 0xC3:   // ret
+    case 0xF3: { // rep; ret
+      result = syms_memread(memread_info, sp, &ip, gpr_size);
+      if (SYMS_RESULT_FAIL(result)) {
+        return result;
+      }
+      sp += gpr_size;
+    } break;
+
+    case 0xe9: { // jmp nnnn
+      S32 imm = syms_buffer_read_s32(&inst_cursor);
+      S64 seek_to = (S64)inst_cursor.off + (S64)imm;
+      if (seek_to < 0) break;
+      keep_parsing = syms_buffer_seek(&inst_cursor, (SymsOffset)seek_to);
+    } break;
+    case 0xeb: { // jmp n
+      S8 imm = syms_buffer_read_s8(&inst_cursor);
+      S64 seek_to = (s64)inst_cursor.off + (s64)imm;
+      if (seek_to < 0) break;
+      keep_parsing = syms_buffer_seek(&inst_cursor, (SymsOffset)seek_to);
+      SYMS_ASSERT_PARANOID(keep_parsing);
+    } break;
+    }
+  }
+
+  regwrite(regwrite_context, arch, ip_regid, &ip, gpr_size);
+  regwrite(regwrite_context, arch, sp_regid, &sp, gpr_size);
+
+  return result;
+}
+
+SYMS_INTERNAL SymsErrorCode
+syms_virtual_unwind_nt(SymsInstance *instance, 
+    SymsMemread *memread_info, 
+    void *regread_context, syms_regread_sig *regread,
+    void *regwrite_context, syms_regwrite_sig *regwrite)
+
+{
+  SymsImage *img = &instance->img;
+  SymsArch arch = syms_get_arch(instance);
+  syms_uint gpr_size = syms_get_addr_size_ex(arch);
+  SymsRegID ip_regid = syms_regs_get_ip_ex(arch);
+  SymsRegID sp_regid = syms_regs_get_sp_ex(arch);
+  SymsAddr ip, sp;
+  syms_uint read_size;
+  SymsNTPdata pdata;
+  SymsAddr frame_base;
+  syms_uint uwinfo_depth_check;
+  syms_bool is_ip_outside_prolog;
+  syms_bool set_frame_base = syms_true;
+  syms_bool machframe_present = syms_false;
+  syms_bool keep_parsing = syms_true;
+  SymsErrorCode result;
+
+  char uw_buffer[128 * sizeof(SymsNTUnwindInfo)]; // TODO: put this on arena
+
+  SYMS_ASSERT(syms_get_addr_size(instance) == 8);
+  //SYMS_ASSERT(img->flags & SYMS_LOAD_IMAGE_FLAGS_FROM_MEMORY);
+
+  ip = 0;
+  read_size = regread(regread_context, arch, ip_regid, &ip, sizeof(ip));
+  if (read_size == 0) {
+    return SYMS_ERR_INREAD;
+  }
+
+  sp = 0;
+  read_size = regread(regread_context, arch, sp_regid, &sp, sizeof(sp));
+  if (read_size == 0) {
+    return SYMS_ERR_INREAD;
+  }
+
+  result = SYMS_ERR_INVALID_CODE_PATH;
+
+  result = syms_find_nearest_pdata(instance, ip, &pdata);
+  if (SYMS_RESULT_FAIL(result)) {
+    // NOTE(nick): Procedure doesn't occupy any stack memory, therefore no pdata entry.
+    SYMS_ASSERT(gpr_size <= sizeof(ip));
+    ip = 0;
+    result = syms_memread(memread_info, sp, &ip, gpr_size);
+    if (SYMS_RESULT_OK(result)) {
+      sp += gpr_size;
+      regwrite(regwrite_context, arch, sp_regid, &sp, gpr_size);
+      regwrite(regwrite_context, arch, ip_regid, &ip, gpr_size);
+    }
+    return result;
+  } else if (result == SYMS_ERR_MAYBE) {
+    return result;
+  }
+
+  is_ip_outside_prolog = syms_true;
+  for (;;) {
+    SymsNTUnwindInfo uwinfo;
+    syms_uint uwinfo_flags;
+    SymsAddr next_uwinfo;
+
+    result = syms_memread(memread_info, pdata.uwinfo, &uwinfo, sizeof(uwinfo));
+    if (SYMS_RESULT_FAIL(result)) {
+      return result;
+    }
+
+    if (ip < pdata.lo + uwinfo.prolog_size) {
+      is_ip_outside_prolog = syms_false;
+      break;
+    }
+
+    uwinfo_flags = SYMS_NT_UNWIND_INFO_HEADER_GET_FLAGS(uwinfo.header);
+    if (~uwinfo_flags & SYMS_NT_UNWIND_INFO_CHAINED) {
+      break;
+    }
+
+    next_uwinfo = pdata.uwinfo + SYMS_NT_UNWIND_INFO_GET_CODE_COUNT(uwinfo.codes_num)*sizeof(SymsNTUnwindInfo);
+    result = syms_memread_pdata(instance, memread_info, next_uwinfo, &pdata);
+    if (SYMS_RESULT_FAIL(result)) {
+      return result;
+    }
+  }
+
+  if (is_ip_outside_prolog) {
+    U8 inst_at_pc[512];
+    SymsBuffer inst_cursor;
+
+    syms_memset(inst_at_pc, 0, sizeof(inst_at_pc));
+    result = syms_memread(memread_info, ip, inst_at_pc, sizeof(inst_at_pc));
+    if (SYMS_RESULT_FAIL(result)) {
+      return result;
+    }
+    inst_cursor = syms_buffer_init(inst_at_pc, sizeof(inst_at_pc));
+    if (syms_is_ip_inside_epilog(inst_cursor, pdata, ip)) {
+      return syms_unwind_from_epilog(img, arch, memread_info, regread_context, regread, regwrite_context, regwrite, ip, sp, inst_cursor);
+    }
+  }
+
+  frame_base = SYMS_ADDR_MAX;
+  uwinfo_depth_check = 0;
+
+  // TODO(nick): Get rid of these bools
+  set_frame_base = syms_true;
+  machframe_present = syms_false;
+  keep_parsing = syms_true;
+
+  for (; keep_parsing; ) {
+    SymsNTUnwindInfo *uwinfo;
+    SymsNTUnwindCode *uwcodes;
+    syms_uint version      = 0;
+    syms_uint uwinfo_size  = 0;
+    syms_uint uwinfo_flags = 0;
+    syms_uint frame_reg    = 0;
+    syms_uint frame_off    = 0;
+    syms_uint code_index   = 0;
+
+    keep_parsing = syms_false;
+
+    result = syms_memread(memread_info, pdata.uwinfo, uw_buffer, sizeof(SymsNTUnwindInfo));
+    if (SYMS_RESULT_FAIL(result)) {
+      return result;
+    }
+    uwinfo = (SymsNTUnwindInfo *)uw_buffer;
+    uwcodes = (SymsNTUnwindCode *)((u8 *)uw_buffer + sizeof(*uwinfo));
+
+    version = SYMS_NT_UNWIND_INFO_HEADER_GET_VERSION(uwinfo->header);
+    if (version != 1 && version != 2) {
+      return SYMS_ERR_INVALID_CODE_PATH;
+    }
+
+    uwinfo_size = syms_nt_unwind_info_sizeof(uwinfo);
+    SYMS_ASSERT(uwinfo_size <= sizeof(uw_buffer));
+    result = syms_memread(memread_info, pdata.uwinfo, uw_buffer, uwinfo_size);
+    if (SYMS_RESULT_FAIL(result)) {
+      return result;
+    }
+
+    uwinfo_flags = SYMS_NT_UNWIND_INFO_HEADER_GET_FLAGS(uwinfo->header);
+    frame_reg = SYMS_NT_UWNIND_INFO_FRAME_GET_REG(uwinfo->frame);
+    frame_off = SYMS_NT_UNWIND_INFO_FRAME_GET_OFF(uwinfo->frame);
+    code_index = 0;
+
+    while (code_index < uwinfo->codes_num) {
+      SymsNTUnwindCode *code = uwcodes + code_index;
+      syms_uint op         = SYMS_NT_UNWIND_CODE_FLAGS_GET_OP_CODE(code->u.flags);
+      syms_uint op_info    = SYMS_NT_UNWIND_CODE_FLAGS_GET_OP_INFO(code->u.flags);
+      syms_uint code_count = syms_nt_unwind_code_count_nodes(code->u.flags);
+      SymsAddr code_ip;
+
+      if (code_count == 0) {
+        break;
+      }
+
+      code_ip = pdata.lo + code->u.off;
+      if (ip >= code_ip) {
+        SymsRegID regid;
+        U64 gpr;
+        U64 off;
+        U8 xmm[16];
+
+        if (set_frame_base) {
+          if (frame_reg != 0) {
+            regid = syms_remap_gpr_nt(img, frame_reg);
+            frame_base = 0;
+            regread(regread_context, arch, regid, &frame_base, sizeof(frame_base));
+            if (frame_base >= frame_off*16u) {
+              frame_base -= frame_off*16u;
+            } else {
+              frame_base = sp;
+            }
+          } else {
+            frame_base = sp;
+          }
+          set_frame_base = syms_false;
+        }
+
+        switch (op) {
+        case SYMS_NT_OP_PUSH_NONVOL: {
+          SYMS_ASSERT(code_count == 1);
+          result = syms_memread(memread_info, sp, &gpr, gpr_size);
+          if (SYMS_RESULT_FAIL(result)) {
+            return syms_false;
+          }
+          sp += 8;
+          regid = syms_remap_gpr_nt(img, op_info);
+          regwrite(regwrite_context, arch, regid, &gpr, gpr_size);
+        } break;
+
+        case SYMS_NT_OP_ALLOC_LARGE: {
+          SYMS_ASSERT(code_count == 2);
+          if (op_info == 0) { // 136 to 512K-8 bytes 
+            sp += code[1].frame_off*8;
+          } else if (op_info == 1) { // 512K to 4G–8 bytes 
+            sp += code[1].frame_off + ((syms_uint)code[2].frame_off << 16);
+          } else {
+            return SYMS_ERR_INVALID_CODE_PATH;
+          }
+        } break;
+
+        case SYMS_NT_OP_ALLOC_SMALL: {
+          SYMS_ASSERT(code_count == 1);
+          sp += op_info*8 + 8;
+        } break;
+
+        case SYMS_NT_OP_SET_FPREG: {
+          SYMS_ASSERT(code_count == 1);
+          regid = syms_remap_gpr_nt(img, frame_reg);
+          // TODO: why are we not using "sp" that we read from begining of the function??
+          sp = 0;
+          regread(regread_context, arch, regid, &sp, gpr_size);
+          SYMS_ASSERT_ALWAYS(sp >= frame_off*16u);
+          sp -= frame_off*16u;
+        } break;
+
+        case SYMS_NT_OP_SAVE_NONVOL: {
+          SYMS_ASSERT(code_count == 2);
+          off = uwcodes[code_index + 1].frame_off*8;
+          result = syms_memread(memread_info, frame_base + off, &gpr, gpr_size);
+          if (SYMS_RESULT_FAIL(result)) {
+            return result;
+          }
+          regid = syms_remap_gpr_nt(img, op_info);
+          regwrite(regwrite_context, arch, regid, &gpr, gpr_size);
+        } break;
+
+        case SYMS_NT_OP_SAVE_NONVOL_FAR: {
+          SYMS_ASSERT(code_count == 3);
+          off = code[0].u.off + ((syms_uint)code[1].u.off << 16);
+          result = syms_memread(memread_info, frame_base + off, &gpr, gpr_size);
+          if (SYMS_RESULT_FAIL(result)) {
+            return result;
+          }
+          regid = syms_remap_gpr_nt(img, op_info);
+          regwrite(regwrite_context, arch, regid, &gpr, gpr_size);
+        } break;
+
+        case SYMS_NT_OP_SAVE_XMM128: {
+          SYMS_ASSERT(code_count == 2);
+          off = code[0].u.off*16;
+          result = syms_memread(memread_info, frame_base + off, xmm, sizeof(xmm));
+          if (SYMS_RESULT_FAIL(result)) {
+            return result;
+          }
+          regid = syms_remap_gpr_nt(img, op_info);
+          regwrite(regwrite_context, arch, regid, xmm, sizeof(xmm));
+        } break;
+
+        case SYMS_NT_OP_SAVE_XMM128_FAR: {
+          SYMS_ASSERT(code_count == 3);
+          off = code[0].frame_off + ((syms_uint)code[0].frame_off << 16);
+          result = syms_memread(memread_info, frame_base + off, xmm, sizeof(xmm));
+          if (SYMS_RESULT_FAIL(result)) {
+            return result;
+          }
+          regid = syms_remap_gpr_nt(img, op_info);
+          regwrite(regwrite_context, arch, regid, &xmm[0], sizeof(xmm));
+        } break;
+
+        case SYMS_NT_OP_EPILOG: {
+          SYMS_ASSERT(code_count == 2);
+        } break;
+
+        case SYMS_NT_OP_PUSH_MACHFRAME: {
+          U16 seg_ss;
+          U64 rflags;
+
+          SYMS_ASSERT(code_count == 1);
+
+          if (op_info > 1) {
+            return SYMS_ERR_INVALID_CODE_PATH;
+          }
+
+          if (op_info == 1) {
+            sp += gpr_size;
+          }
+
+          result = syms_memread(memread_info, sp, &ip, gpr_size);
+          if (SYMS_RESULT_FAIL(result)) {
+            return result;
+          }
+          sp += gpr_size;
+
+          seg_ss = 0;
+          result = syms_memread(memread_info, sp, &seg_ss, 2);
+          if (SYMS_RESULT_FAIL(result)) {
+            return result;
+          }
+          regwrite(regwrite_context, arch, SYMS_REG_X64_ss, &seg_ss, sizeof(seg_ss));
+          sp += gpr_size;
+
+          rflags = 0;
+          result = syms_memread(memread_info, sp, &rflags, gpr_size);
+          if (SYMS_RESULT_FAIL(result)) {
+            return result;
+          }
+          regwrite(regwrite_context, arch, SYMS_REG_X64_rflags, &rflags, sizeof(rflags));
+          sp += gpr_size;
+
+          result = syms_memread(memread_info, sp + gpr_size, &seg_ss, sizeof(seg_ss));
+          if (SYMS_RESULT_FAIL(result)) {
+            return result;
+          }
+
+          result = syms_memread(memread_info, sp, &sp, gpr_size);
+          if (SYMS_RESULT_FAIL(result)) {
+            return result;
+          }
+
+          machframe_present = syms_true;
+        } break;
+
+        default: SYMS_ASSERT_FAILURE("syms_stack_unwind_frame: unknown NT unwind op");
+        }
+      }
+
+      code_index += code_count;
+    }
+
+    if (uwinfo_flags & SYMS_NT_UNWIND_INFO_CHAINED) {
+      U8 last_code_index = SYMS_NT_UNWIND_INFO_GET_CODE_COUNT(uwinfo->codes_num);
+      void *next_pdata = (void *)&uwcodes[last_code_index];
+      pdata = syms_unpack_pdata(instance, (struct SymsNTPdataPacked *)next_pdata);
+      if (++uwinfo_depth_check > 32) {
+        return SYMS_ERR_INVALID_CODE_PATH;
+      }
+      keep_parsing = syms_true;
+    }
+  }
+
+  if (!machframe_present) {
+    result = syms_memread(memread_info, sp, &ip, gpr_size);
+    if (SYMS_RESULT_FAIL(result)) {
+      return result;
+    }
+    sp += gpr_size;
+  }
+
+  if (ip == 0) {
+    // When ip is 0 the stack has been fully unwound.
+    return SYMS_ERR_INVALID_CODE_PATH;
+  }
+
+  regwrite(regwrite_context, arch, ip_regid, &ip, gpr_size);
+  regwrite(regwrite_context, arch, sp_regid, &sp, gpr_size);
+
+  return SYMS_ERR_OK;
+}
diff --git a/Engine/Source/Runtime/SymsLib/src/win/syms_nt_unwind.h b/Engine/Source/Runtime/SymsLib/src/win/syms_nt_unwind.h
new file mode 100644
--- /dev/null
+++ b/Engine/Source/Runtime/SymsLib/src/win/syms_nt_unwind.h
@@ -0,0 +1,31 @@
+// Copyright Epic Games, Inc. All Rights Reserved.
+#ifndef SYMS_NT_UNWIND_INCLUDE_H
+#define SYMS_NT_UNWIND_INCLUDE_H
+
+/******************************************************************************
+ * File   : syms_nt_unwind.h                                                  *
+ * Author : Nikita Smith                                                      *
+ * Created: 2020/09/11                                                        *
+ * Purpose: Stack unwinder for Win32                                          *
+ ******************************************************************************/
+
+SYMS_INTERNAL syms_uint
+syms_nt_unwind_info_sizeof(SymsNTUnwindInfo *uwinfo);
+
+SYMS_INTERNAL syms_uint
+syms_nt_unwind_code_count_nodes(U8 uwcode_flags);
+
+SYMS_INTERNAL SymsRegID
+syms_remap_gpr_nt(SymsImage *img, syms_uint nt_regid);
+
+SYMS_INTERNAL SymsRegID
+syms_remap_xmm_nt(SymsImage *img, syms_uint nt_regid);
+
+SYMS_INTERNAL SymsErrorCode
+syms_memread_pdata(struct SymsInstance *instance, SymsMemread *memread_info, SymsAddr va, SymsNTPdata *pdata_out);
+
+SYMS_INTERNAL SymsErrorCode
+syms_virtual_unwind_nt(struct SymsInstance *instance, SymsRegs *regs, SymsMemread *memread_info);
+
+#endif // SYMS_NT_UNWIND_INCLUDE_H
+
