Update syms library to latest version, includes dwarf5 support

diff --git a/Engine/Source/Runtime/SymsLib/syms/core/group/syms_group.c b/Engine/Source/Runtime/SymsLib/syms/core/group/syms_group.c
new file mode 100644
--- /dev/null
+++ b/Engine/Source/Runtime/SymsLib/syms/core/group/syms_group.c
@@ -0,0 +1,2483 @@
+// Copyright Epic Games, Inc. All Rights Reserved.
+#ifndef SYMS_GROUP_C
+#define SYMS_GROUP_C
+
+// TODO(allen): killing the arena copy optimizations may cause some problems
+// in group, take a closer look into usage of string cons again.
+
+////////////////////////////////
+// NOTE(allen): Syms Group Setup Functions
+
+SYMS_API SYMS_Group*
+syms_group_alloc(void){
+  SYMS_Arena *arena = syms_arena_alloc();
+  SYMS_Group *result = syms_push_array_zero(arena, SYMS_Group, 1);
+  result->arena = arena;
+  return(result);
+}
+
+SYMS_API void
+syms_group_release(SYMS_Group *group){
+  syms_arena_release(group->arena);
+}
+
+SYMS_API void
+syms_group_init(SYMS_Group *group, SYMS_ParseBundle *params){
+  SYMS_ProfBegin("syms_group_init");
+  SYMS_ASSERT(syms_thread_lane == 0);
+  
+  //- fill in bin
+  SYMS_String8 bin_data = params->bin_data;
+  SYMS_BinAccel *bin = params->bin;
+  if (bin_data.size > 0){
+    if (bin == 0){
+      SYMS_FileAccel *bin_file = syms_file_accel_from_data(group->arena, params->bin_data);
+      bin = syms_bin_accel_from_file(group->arena, params->bin_data, bin_file);
+    }
+  }
+  else{
+    bin = (SYMS_BinAccel*)&syms_format_nil;
+  }
+  
+  //- fill in dbg
+  SYMS_String8 dbg_data = params->dbg_data;
+  SYMS_DbgAccel *dbg = params->dbg;
+  if (dbg_data.size > 0){
+    if (dbg == 0){
+      SYMS_FileAccel *dbg_file = syms_file_accel_from_data(group->arena, params->dbg_data);
+      dbg = syms_dbg_accel_from_file(group->arena, params->dbg_data, dbg_file);
+    }
+  }
+  else{
+    dbg = (SYMS_DbgAccel*)&syms_format_nil;
+  }
+  
+  //- arch & address size
+  SYMS_Arch arch = syms_arch_from_bin(bin);
+  if (arch == SYMS_Arch_Null){
+    arch = syms_arch_from_dbg(dbg);
+  }
+  SYMS_U64 address_size = (syms_address_size_from_arch(arch) >> 3);
+  
+  //- default vbase
+  SYMS_U64 default_vbase = 0;
+  if (bin->format != SYMS_FileFormat_Null){
+    default_vbase = syms_default_vbase_from_bin(bin);
+  }
+  else{
+    default_vbase = syms_default_vbase_from_dbg(dbg);
+  }
+  
+  //- sections
+  SYMS_SecInfoArray sec_info_array = {0};
+  if (bin->format != SYMS_FileFormat_Null){
+    sec_info_array = syms_sec_info_array_from_bin(group->arena, bin_data, bin);
+  }
+  else{
+    sec_info_array = syms_sec_info_array_from_dbg(group->arena, dbg_data, dbg);
+  }
+  
+  //- top-level accelerators
+  SYMS_UnitSetAccel *unit_set = syms_unit_set_accel_from_dbg(group->arena, dbg_data, dbg);
+  SYMS_U64 unit_count = syms_unit_count_from_set(unit_set);
+  SYMS_MapAccel *type_map = syms_type_map_from_dbg(group->arena, dbg_data, dbg);
+  
+  //- init lanes to single-threaded setup
+  group->lane_arenas = &group->arena;
+  group->lane_count = 1;
+  
+  //- fill the group
+  group->bin_data = bin_data;
+  group->bin = bin;
+  group->dbg_data = dbg_data;
+  group->dbg = dbg;
+  
+  group->arch = arch;
+  group->address_size = address_size;
+  group->default_vbase = default_vbase;
+  group->sec_info_array = sec_info_array;
+  group->unit_set = unit_set;
+  group->unit_count = unit_count;
+  group->type_map = type_map;
+  
+  //- setup basic caches
+  group->unit_cache_flags = syms_push_array_zero(group->arena, SYMS_GroupUnitCacheFlags, unit_count);
+  group->units = syms_push_array_zero(group->arena, SYMS_UnitAccel*, unit_count);
+  group->proc_sid_arrays = syms_push_array_zero(group->arena, SYMS_SymbolIDArray, unit_count);
+  group->var_sid_arrays = syms_push_array_zero(group->arena, SYMS_SymbolIDArray, unit_count);
+  group->type_sid_arrays = syms_push_array_zero(group->arena, SYMS_SymbolIDArray, unit_count);
+  group->file_tables = syms_push_array_zero(group->arena, SYMS_String8Array, unit_count);
+  group->inferred_file_tables = syms_push_array_zero(group->arena, SYMS_String8Array, unit_count);
+  group->line_tables = syms_push_array_zero(group->arena, SYMS_LineParseOut, unit_count);
+  group->unit_proc_maps = syms_push_array_zero(group->arena, SYMS_SpatialMap1D, unit_count);
+  group->unit_var_maps = syms_push_array_zero(group->arena, SYMS_SpatialMap1D, unit_count);
+  group->line_sequence_maps = syms_push_array_zero(group->arena, SYMS_SpatialMap1D, unit_count);
+  group->file_to_line_to_addr_buckets = syms_push_array_zero(group->arena, SYMS_FileToLineToAddrBuckets,
+                                                             unit_count);
+  
+  //- setup hash table caches
+  group->string_cons = syms_string_cons_alloc(group->arena, 4093);
+  group->file_id_2_name_map = syms_file_id_2_name_map_alloc(group->arena, 4093);
+  
+  group->type_usid_buckets.bucket_count = 4093;
+  group->type_usid_buckets.buckets = syms_push_array_zero(group->arena, SYMS_TypeUSIDNode*, 4093);
+  
+  group->type_content_buckets.bucket_count = 4093;
+  group->type_content_buckets.buckets = syms_push_array_zero(group->arena, SYMS_TypeContentNode*, 4093);
+  
+  SYMS_ProfEnd();
+}
+
+SYMS_API void
+syms_group_parse_all_units__single_thread(SYMS_Group *group){
+  SYMS_ProfBegin("syms_group_parse_all_units__single_thread");
+  
+  SYMS_U64 count = group->unit_count;
+  for (SYMS_U64 n = 1; n <= count; n += 1){
+    SYMS_UnitID uid = n;
+    syms_group_unit_from_uid(group, uid);
+    syms_group_proc_sid_array_from_uid(group, uid);
+    syms_group_var_sid_array_from_uid(group, uid);
+    syms_group_type_sid_array_from_uid(group, uid);
+    syms_group_file_table_from_uid(group, uid);
+    syms_group_line_table_from_uid(group, uid);
+    syms_group_proc_map_from_uid(group, uid);
+    syms_group_line_sequence_map_from_uid(group, uid);
+    syms_group_fetch_line_to_addr_maps_from_uid(group, uid);
+  }
+  
+  SYMS_ProfEnd();
+}
+
+SYMS_API void
+syms_group_parse_all_top_level(SYMS_Group *group){
+  SYMS_ProfBegin("syms_group_parse_all_top_level");
+  
+  syms_group_unit_ranges(group);
+  syms_group_type_map(group);
+  
+  syms_group_sec_map_v(group);
+  syms_group_sec_map_f(group);
+  syms_group_unit_map(group);
+  syms_group_file_map(group);
+  
+  SYMS_ProfEnd();
+}
+
+SYMS_API void
+syms_group_begin_multilane(SYMS_Group *group, SYMS_U64 lane_count){
+  SYMS_ASSERT(syms_thread_lane == 0);
+  SYMS_ASSERT(group->lane_count == 1);
+  if (lane_count > 1){
+    // allocate the group's lane arenas array
+    if (lane_count > group->lane_max){
+      group->lane_arenas = syms_push_array(group->arena, SYMS_Arena*, lane_count);
+      group->lane_max = lane_count;
+    }
+    
+    // fill the group's lane arenas array
+    SYMS_Arena **lane_arena_ptr = group->lane_arenas;
+    *lane_arena_ptr = group->arena;
+    lane_arena_ptr += 1;
+    for (SYMS_U64 i = 1; i < lane_count; i += 1, lane_arena_ptr += 1){
+      *lane_arena_ptr = syms_arena_alloc();
+    }
+    group->lane_count = lane_count;
+  }
+}
+
+SYMS_API void
+syms_group_end_multilane(SYMS_Group *group){
+  SYMS_U64 lane_count = group->lane_count;
+  if (lane_count > 1){
+    // absorb the lane arenas
+    SYMS_Arena *arena = group->arena;
+    SYMS_Arena **lane_arena_ptr = group->lane_arenas + 1;
+    for (SYMS_U64 i = 1; i < lane_count; i += 1, lane_arena_ptr += 1){
+      syms_arena_absorb(arena, *lane_arena_ptr);
+      *lane_arena_ptr = 0;
+    }
+    
+    // drop the lane count
+    group->lane_count = 1;
+  }
+}
+
+SYMS_API SYMS_Arena*
+syms_group_get_lane_arena(SYMS_Group *group){
+  SYMS_U64 thread_lane = syms_get_lane();
+  SYMS_Arena *result = 0;
+  if (thread_lane < group->lane_count){
+    result = group->lane_arenas[thread_lane];
+  }
+  return(result);
+}
+
+////////////////////////////////
+// NOTE(allen): Syms Group Getters
+
+SYMS_API SYMS_String8
+syms_group_bin_data(SYMS_Group *group){
+  return(group->bin_data);
+}
+
+SYMS_API SYMS_BinAccel*
+syms_group_bin(SYMS_Group *group){
+  return(group->bin);
+}
+
+SYMS_API SYMS_String8
+syms_group_dbg_data(SYMS_Group *group){
+  return(group->dbg_data);
+}
+
+SYMS_API SYMS_DbgAccel*
+syms_group_dbg(SYMS_Group *group){
+  return(group->dbg);
+}
+
+SYMS_API SYMS_UnitSetAccel*
+syms_group_unit_set(SYMS_Group *group){
+  return(group->unit_set);
+}
+
+SYMS_API SYMS_SecInfoArray
+syms_group_sec_info_array(SYMS_Group *group){
+  SYMS_SecInfoArray result = group->sec_info_array;
+  return(result);
+}
+
+SYMS_API SYMS_SecInfo*
+syms_group_sec_info_from_number(SYMS_Group *group, SYMS_U64 n){
+  SYMS_ProfBegin("syms_group_sec_info_from_number");
+  SYMS_SecInfoArray array = group->sec_info_array;
+  SYMS_SecInfo *result = 0;
+  if (1 <= n && n <= array.count){
+    result = &array.sec_info[n - 1];
+  }
+  SYMS_ProfEnd();
+  return(result);
+}
+
+SYMS_API SYMS_U64
+syms_group_unit_count(SYMS_Group *group){
+  return(group->unit_count);
+}
+
+SYMS_API SYMS_UnitInfo
+syms_group_unit_info_from_uid(SYMS_Group *group, SYMS_UnitID uid){
+  SYMS_UnitInfo result = syms_unit_info_from_uid(group->unit_set, uid);
+  return(result);
+}
+
+SYMS_API SYMS_UnitNames
+syms_group_unit_names_from_uid(SYMS_Arena *arena, SYMS_Group *group, SYMS_UnitID uid){
+  SYMS_UnitNames result = syms_unit_names_from_uid(arena, group->unit_set, uid);
+  return(result);
+}
+
+SYMS_API SYMS_UnitAccel*
+syms_group_unit_from_uid(SYMS_Group *group, SYMS_UnitID uid){
+  SYMS_ProfBegin("syms_group_unit_from_uid");
+  SYMS_Arena *arena = syms_group_get_lane_arena(group);
+  SYMS_UnitAccel *result = (SYMS_UnitAccel*)&syms_format_nil;
+  if (1 <= uid && uid <= group->unit_count){
+    SYMS_U64 index = uid - 1;
+    result = group->units[index];
+    if (result == 0){
+      SYMS_UnitAccel *unit = syms_unit_accel_from_uid(arena, group->dbg_data, group->dbg, group->unit_set, uid);
+      group->units[index] = unit;
+      result = unit;
+    }
+  }
+  SYMS_ProfEnd();
+  return(result);
+}
+
+SYMS_API SYMS_SymbolIDArray*
+syms_group_proc_sid_array_from_uid(SYMS_Group *group, SYMS_UnitID uid){
+  SYMS_ProfBegin("syms_group_proc_sid_array_from_uid");
+  SYMS_SymbolIDArray *result = &syms_sid_array_nil;
+  if (1 <= uid && uid <= group->unit_count){
+    SYMS_U64 index = uid - 1;
+    if ((group->unit_cache_flags[index] & SYMS_GroupUnitCacheFlag_HasProcSidArray) == 0){
+      SYMS_UnitAccel *unit = syms_group_unit_from_uid(group, uid);
+      SYMS_Arena *arena = syms_group_get_lane_arena(group);
+      group->proc_sid_arrays[index] = syms_proc_sid_array_from_unit(arena, unit);
+      group->unit_cache_flags[index] |= SYMS_GroupUnitCacheFlag_HasProcSidArray;
+    }
+    result = &group->proc_sid_arrays[index];
+  }
+  SYMS_ProfEnd();
+  return(result);
+}
+
+SYMS_API SYMS_SymbolIDArray*
+syms_group_var_sid_array_from_uid(SYMS_Group *group, SYMS_UnitID uid){
+  SYMS_ProfBegin("syms_group_var_sid_array_from_uid");
+  SYMS_SymbolIDArray *result = &syms_sid_array_nil;
+  if (1 <= uid && uid <= group->unit_count){
+    SYMS_U64 index = uid - 1;
+    if ((group->unit_cache_flags[index] & SYMS_GroupUnitCacheFlag_HasVarSidArray) == 0){
+      SYMS_UnitAccel *unit = syms_group_unit_from_uid(group, uid);
+      SYMS_Arena *arena = syms_group_get_lane_arena(group);
+      group->var_sid_arrays[index] = syms_var_sid_array_from_unit(arena, unit);
+      group->unit_cache_flags[index] |= SYMS_GroupUnitCacheFlag_HasVarSidArray;
+    }
+    result = &group->var_sid_arrays[index];
+  }
+  SYMS_ProfEnd();
+  return(result);
+}
+
+SYMS_API SYMS_SymbolIDArray*
+syms_group_type_sid_array_from_uid(SYMS_Group *group, SYMS_UnitID uid){
+  SYMS_ProfBegin("syms_group_type_sid_array_from_uid");
+  SYMS_SymbolIDArray *result = &syms_sid_array_nil;
+  if (1 <= uid && uid <= group->unit_count){
+    SYMS_U64 index = uid - 1;
+    if ((group->unit_cache_flags[index] & SYMS_GroupUnitCacheFlag_HasTypeSidArray) == 0){
+      SYMS_UnitAccel *unit = syms_group_unit_from_uid(group, uid);
+      SYMS_Arena *arena = syms_group_get_lane_arena(group);
+      group->type_sid_arrays[index] = syms_type_sid_array_from_unit(arena, unit);
+      group->unit_cache_flags[index] |= SYMS_GroupUnitCacheFlag_HasTypeSidArray;
+    }
+    result = &group->type_sid_arrays[index];
+  }
+  SYMS_ProfEnd();
+  return(result);
+}
+
+SYMS_API SYMS_String8Array*
+syms_group_file_table_from_uid(SYMS_Group *group, SYMS_UnitID uid){
+  SYMS_ProfBegin("syms_group_file_table_from_uid");
+  SYMS_String8Array *result = &syms_string_array_nil;
+  if (1 <= uid && uid <= group->unit_count){
+    SYMS_U64 index = uid - 1;
+    if ((group->unit_cache_flags[index] & SYMS_GroupUnitCacheFlag_HasFileTable) == 0){
+      SYMS_Arena *arena = syms_group_get_lane_arena(group);
+      group->file_tables[index] = syms_file_table_from_uid(arena, group->dbg_data, group->dbg,
+                                                           group->unit_set, uid);
+      group->unit_cache_flags[index] |= SYMS_GroupUnitCacheFlag_HasFileTable;
+    }
+    result = &group->file_tables[index];
+  }
+  SYMS_ProfEnd();
+  return(result);
+}
+
+SYMS_API SYMS_LineParseOut*
+syms_group_line_parse_from_uid(SYMS_Group *group, SYMS_UnitID uid){
+  SYMS_ProfBegin("syms_group_line_parse_from_uid");
+  SYMS_LineParseOut* result = &syms_line_parse_nil;
+  if (1 <= uid && uid <= group->unit_count){
+    SYMS_U64 index = uid - 1;
+    if ((group->unit_cache_flags[index] & SYMS_GroupUnitCacheFlag_HasLineTable) == 0){
+      SYMS_Arena *arena = syms_group_get_lane_arena(group);
+      group->line_tables[index] = syms_line_parse_from_uid(arena, group->dbg_data, group->dbg, group->unit_set, uid);
+      group->unit_cache_flags[index] |= SYMS_GroupUnitCacheFlag_HasLineTable;
+    }
+    result = &group->line_tables[index];
+  }
+  SYMS_ProfEnd();
+  return(result);
+}
+
+SYMS_API SYMS_LineTable*
+syms_group_line_table_from_uid(SYMS_Group *group, SYMS_UnitID uid){
+  SYMS_ProfBegin("syms_group_line_table_from_uid");
+  SYMS_LineParseOut *parse = syms_group_line_parse_from_uid(group, uid);
+  SYMS_LineTable *result = &parse->line_table;
+  SYMS_ProfEnd();
+  return(result);
+}
+
+SYMS_API SYMS_String8Array*
+syms_group_inferred_file_table_from_uid(SYMS_Group *group, SYMS_UnitID uid){
+  SYMS_ProfBegin("syms_group_inferred_file_table_from_uid");
+  SYMS_String8Array *result = &syms_string_array_nil;
+  if (1 <= uid && uid <= group->unit_count){
+    SYMS_U32 index = uid - 1;
+    if ((group->unit_cache_flags[index] & SYMS_GroupUnitCacheFlag_HasInfFileTable) == 0){
+      // get the line parse
+      SYMS_LineParseOut *line_parse = syms_group_line_parse_from_uid(group, uid);
+      
+      if (line_parse->file_table.count != 0){
+        result = &line_parse->file_table;
+      }
+      else{
+        // setup
+        SYMS_Arena *arena = syms_group_get_lane_arena(group);
+        SYMS_String8 dbg_data = group->dbg_data;
+        SYMS_DbgAccel *dbg = group->dbg;
+        SYMS_UnitSetAccel *unit_set = group->unit_set;
+        
+        // infer names
+        SYMS_U64 count = line_parse->file_id_array.count;
+        SYMS_String8 *names = syms_push_array(arena, SYMS_String8, count);
+        {
+          SYMS_FileID *file_id_ptr = line_parse->file_id_array.ids;
+          SYMS_String8 *name_ptr = names;
+          for (SYMS_FileID *opl = file_id_ptr + count;
+               file_id_ptr < opl;
+               file_id_ptr += 1, name_ptr += 1){
+            *name_ptr = syms_file_name_from_id(arena, dbg_data, dbg, unit_set, uid, *file_id_ptr);
+          }
+        }
+        
+        // assemble file table
+        SYMS_String8Array file_table = {0};
+        file_table.strings = names;
+        file_table.count = count;
+        
+        // fill cache slot
+        group->inferred_file_tables[index] = file_table;
+        group->unit_cache_flags[index] |= SYMS_GroupUnitCacheFlag_HasInfFileTable;
+      }
+    }
+    result = &group->inferred_file_tables[index];
+  }
+  SYMS_ProfEnd();
+  return(result);
+}
+
+SYMS_API SYMS_String8Array*
+syms_group_file_table_from_uid_with_fallbacks(SYMS_Group *group, SYMS_UnitID uid){
+  SYMS_ProfBegin("syms_group_file_table_from_uid_with_fallbacks");
+  SYMS_String8Array *result = syms_group_file_table_from_uid(group, uid);
+  if (result != 0 && result->count == 0){
+    result = syms_group_inferred_file_table_from_uid(group, uid);
+  }
+  SYMS_ProfEnd();
+  return(result);
+}
+
+SYMS_API SYMS_UnitRangeArray
+syms_group_unit_ranges(SYMS_Group *group){
+  if (!group->unit_ranges_is_filled){
+    SYMS_ASSERT(syms_thread_lane == 0);
+    group->unit_ranges = syms_unit_ranges_from_set(group->arena, group->dbg_data, group->dbg, group->unit_set);
+    group->unit_ranges_is_filled = syms_true;
+  }
+  SYMS_UnitRangeArray result = group->unit_ranges;
+  return(result);
+}
+
+SYMS_API SYMS_SymbolKind
+syms_group_symbol_kind_from_sid(SYMS_Group *group, SYMS_UnitAccel *unit, SYMS_SymbolID sid){
+  SYMS_SymbolKind result = syms_symbol_kind_from_sid(group->dbg_data, group->dbg, unit, sid);
+  return(result);
+}
+
+SYMS_API SYMS_String8
+syms_group_symbol_name_from_sid(SYMS_Arena *arena, SYMS_Group *group, SYMS_UnitAccel *unit, SYMS_SymbolID sid){
+  SYMS_String8 result = syms_symbol_name_from_sid(arena, group->dbg_data, group->dbg, unit, sid);
+  return(result);
+}
+
+SYMS_API SYMS_U64
+syms_group_address_size(SYMS_Group *group){
+  SYMS_U64 result = group->address_size;
+  return(result);
+}
+
+SYMS_API SYMS_U64
+syms_group_default_vbase(SYMS_Group *group){
+  SYMS_U64 result = group->default_vbase;
+  return(result);
+}
+
+SYMS_API SYMS_String8
+syms_group_file_name_from_id(SYMS_Group *group, SYMS_UnitID uid, SYMS_FileID file_id){
+  SYMS_ProfBegin("syms_group_file_name_from_id");
+  SYMS_ASSERT(syms_thread_lane == 0);
+  SYMS_String8 from_map = syms_file_id_2_name_map_name_from_id(&group->file_id_2_name_map, uid, file_id);
+  SYMS_String8 result = {0};
+  if (from_map.size != 0){
+    if (from_map.str != 0){
+      result = from_map;
+    }
+  }
+  else{
+    SYMS_ArenaTemp scratch = syms_get_scratch(0, 0);
+    SYMS_String8 file_name = syms_file_name_from_id(scratch.arena, group->dbg_data, group->dbg,
+                                                    group->unit_set, uid, file_id);
+    result = syms_string_cons(group->arena, &group->string_cons, file_name);
+    syms_release_scratch(scratch);
+    syms_file_id_2_name_map_insert(group->arena, &group->file_id_2_name_map, uid, file_id, result);
+  }
+  SYMS_ProfEnd();
+  return(result);
+}
+
+SYMS_API SYMS_MapAndUnit
+syms_group_type_map(SYMS_Group *group){
+  if (!group->type_map_unit_is_filled){
+    SYMS_UnitID type_map_uid = syms_partner_uid_from_map(group->type_map);
+    SYMS_UnitAccel *type_map_unit = syms_group_unit_from_uid(group, type_map_uid);
+    group->type_map_unit = type_map_unit;
+    group->type_map_unit_is_filled = syms_true;
+  }
+  SYMS_MapAndUnit result = {0};
+  result.map = group->type_map;
+  result.unit = group->type_map_unit;
+  return(result);
+}
+
+
+////////////////////////////////
+// NOTE(allen): Syms Group Address Mapping Functions
+
+SYMS_API SYMS_U64
+syms_group_sec_number_from_voff__linear_scan(SYMS_Group *group, SYMS_U64 voff){
+  SYMS_ProfBegin("syms_group_sec_number_from_voff__linear_scan");
+  SYMS_U64 result = 0;
+  SYMS_SecInfoArray sec_info_array = syms_group_sec_info_array(group);
+  SYMS_SecInfo *sec_info = sec_info_array.sec_info;
+  for (SYMS_U64 i = 0; i < sec_info_array.count; i += 1, sec_info += 1){
+    if (sec_info->vrange.min <= voff && voff < sec_info->vrange.max){
+      result = i + 1;
+      break;
+    }
+  }
+  SYMS_ProfEnd();
+  return(result);
+}
+
+SYMS_API SYMS_U64
+syms_group_sec_number_from_foff__linear_scan(SYMS_Group *group, SYMS_U64 foff){
+  SYMS_ProfBegin("syms_group_sec_number_from_foff__linear_scan");
+  SYMS_U64 result = 0;
+  SYMS_SecInfoArray sec_info_array = syms_group_sec_info_array(group);
+  SYMS_SecInfo *sec_info = sec_info_array.sec_info;
+  for (SYMS_U64 i = 0; i < sec_info_array.count; i += 1, sec_info += 1){
+    if (sec_info->frange.min <= foff && foff < sec_info->frange.max){
+      result = i + 1;
+      break;
+    }
+  }
+  SYMS_ProfEnd();
+  return(result);
+}
+
+SYMS_API SYMS_U64Maybe
+syms_group_voff_from_foff__linear_scan(SYMS_Group *group, SYMS_U64 foff){
+  SYMS_ProfBegin("syms_group_voff_from_foff__linear_scan");
+  SYMS_U64Maybe result = {0};
+  SYMS_U64 sec_number = syms_group_sec_number_from_foff__linear_scan(group, foff);
+  SYMS_SecInfo *info = syms_group_sec_info_from_number(group, sec_number);
+  if (info != 0){
+    result.valid = syms_true;
+    result.u64 = foff - info->frange.min + info->vrange.min;
+  }
+  SYMS_ProfEnd();
+  return(result);
+}
+
+SYMS_API SYMS_U64Maybe
+syms_group_foff_from_voff__linear_scan(SYMS_Group *group, SYMS_U64 voff){
+  SYMS_ProfBegin("syms_group_foff_from_voff__linear_scan");
+  SYMS_U64Maybe result = {0};
+  SYMS_U64 sec_number = syms_group_sec_number_from_voff__linear_scan(group, voff);
+  SYMS_SecInfo *info = syms_group_sec_info_from_number(group, sec_number);
+  if (info != 0){
+    result.valid = syms_true;
+    result.u64 = voff - info->vrange.min + info->frange.min;
+  }
+  SYMS_ProfEnd();
+  return(result);
+}
+
+SYMS_API SYMS_UnitID
+syms_group_uid_from_voff__linear_scan(SYMS_Group *group, SYMS_U64 virt_off){
+  SYMS_ProfBegin("syms_group_uid_from_voff__linear_scan");
+  SYMS_UnitID result = 0;
+  SYMS_UnitRangeArray unit_ranges = syms_group_unit_ranges(group);
+  SYMS_U64 range_count = unit_ranges.count;
+  SYMS_UnitRange *unit_range = unit_ranges.ranges;
+  for (SYMS_U64 i = 0; i < range_count; i += 1, unit_range += 1){
+    if (unit_range->vrange.min <= virt_off && virt_off < unit_range->vrange.max){
+      result = unit_range->uid;
+      break;
+    }
+  }
+  SYMS_ProfEnd();
+  return(result);
+}
+
+SYMS_API SYMS_SymbolID
+syms_group_proc_sid_from_uid_voff__linear_scan(SYMS_Group *group, SYMS_UnitID uid, SYMS_U64 voff){
+  SYMS_ProfBegin("syms_group_proc_sid_from_uid_voff__linear_scan");
+  SYMS_ArenaTemp scratch = syms_get_scratch(0, 0);
+  
+  SYMS_String8 dbg_data = group->dbg_data;
+  SYMS_DbgAccel *dbg = group->dbg;
+  SYMS_UnitAccel *unit = syms_group_unit_from_uid(group, uid);
+  SYMS_SymbolIDArray *sid_array = syms_group_proc_sid_array_from_uid(group, uid);
+  SYMS_U64 count = sid_array->count;
+  
+  SYMS_SymbolID result = {0};
+  SYMS_SymbolID *sid_ptr = sid_array->ids;
+  for (SYMS_U64 i = 0; i < count; i += 1, sid_ptr += 1){
+    SYMS_SymbolID sid = *sid_ptr;
+    SYMS_SymbolKind kind = syms_group_symbol_kind_from_sid(group, unit, sid);
+    SYMS_B32 hit = syms_false;
+    if (kind == SYMS_SymbolKind_Procedure){
+      SYMS_U64RangeArray ranges = syms_proc_vranges_from_sid(scratch.arena, dbg_data, dbg, unit, sid);
+      SYMS_U64Range *range = ranges.ranges;
+      for (SYMS_U64 i = 0; i < ranges.count; i += 1, range += 1){
+        if (range->min <= voff && voff < range->max){
+          hit = syms_true;
+          break;
+        }
+      }
+      syms_arena_temp_end(scratch);
+    }
+    if (hit){
+      result = sid;
+      break;
+    }
+  }
+  
+  syms_release_scratch(scratch);
+  SYMS_ProfEnd();
+  return(result);
+}
+
+SYMS_API SYMS_SymbolID
+syms_group_var_sid_from_uid_voff__linear_scan(SYMS_Group *group, SYMS_UnitID uid, SYMS_U64 voff){
+  SYMS_ProfBegin("syms_group_var_sid_from_uid_voff__linear_scan");
+  
+  SYMS_String8 dbg_data = group->dbg_data;
+  SYMS_DbgAccel *dbg = group->dbg;
+  SYMS_UnitAccel *unit = syms_group_unit_from_uid(group, uid);
+  SYMS_SymbolIDArray *sid_array = syms_group_var_sid_array_from_uid(group, uid);
+  SYMS_U64 count = sid_array->count;
+  
+  SYMS_SymbolID result = {0};
+  SYMS_SymbolID *sid_ptr = sid_array->ids;
+  for (SYMS_U64 i = 0; i < count; i += 1, sid_ptr += 1){
+    SYMS_SymbolID sid = *sid_ptr;
+    SYMS_SymbolKind kind = syms_group_symbol_kind_from_sid(group, unit, sid);
+    SYMS_B32 hit = syms_false;
+    if (kind == SYMS_SymbolKind_ImageRelativeVariable){
+      SYMS_U64 var_voff = syms_voff_from_var_sid(dbg_data, dbg, unit, sid);
+      if (var_voff == voff){
+        hit = syms_true;
+      }
+      else if (var_voff < voff){
+        SYMS_USID var_type_usid = syms_type_from_var_sid(dbg_data, dbg, unit, sid);
+        SYMS_U64 var_size = syms_group_type_size_from_usid(group, var_type_usid);
+        if (voff < var_voff + var_size){
+          hit = syms_true;
+        }
+      }
+    }
+    if (hit){
+      result = sid;
+      break;
+    }
+  }
+  
+  SYMS_ProfEnd();
+  return(result);
+}
+
+SYMS_API SYMS_Line
+syms_group_line_from_uid_voff__linear_scan(SYMS_Group *group, SYMS_UnitID uid, SYMS_U64 voff){
+  SYMS_ProfBegin("syms_group_line_from_uid_voff__linear_scan");
+  SYMS_Line result = {0};
+  SYMS_LineTable *line_table = syms_group_line_table_from_uid(group, uid);
+  SYMS_U64 *seq_index_ptr = line_table->sequence_index_array;
+  SYMS_Line *line_array = line_table->line_array;
+  SYMS_U64 seq_count = line_table->sequence_count;
+  for (SYMS_U64 i = 0; i < seq_count; i += 1){
+    // get sequence range & inc
+    SYMS_U64 first = *seq_index_ptr;
+    seq_index_ptr += 1;
+    SYMS_U64 opl = *seq_index_ptr;
+    
+    // check sequence range
+    SYMS_U64 last = opl - 1;
+    SYMS_U64 first_voff = line_array[first].voff;
+    SYMS_U64 opl_voff = line_array[last].voff;
+    if (first_voff <= voff && voff < opl_voff){
+      
+      // get the line
+      result = syms_line_from_sequence_voff(line_table, i + 1, voff);
+    }
+  }
+  SYMS_ProfEnd();
+  return(result);
+}
+
+SYMS_API SYMS_U64RangeList
+syms_group_vranges_from_uid_line__linear_scan(SYMS_Arena *arena, SYMS_Group *group, SYMS_UnitID uid,
+                                              SYMS_FileID file_id, SYMS_U32 line){
+  SYMS_ProfBegin("syms_group_vranges_from_uid_line__linear_scan");
+  
+  SYMS_ASSERT(arena != group->arena);
+  SYMS_U64RangeList result = {0};
+  SYMS_LineTable *line_table = syms_group_line_table_from_uid(group, uid);
+  SYMS_U64 *seq_index_ptr = line_table->sequence_index_array;
+  SYMS_Line *line_array = line_table->line_array;
+  SYMS_U64 seq_count = line_table->sequence_count;
+  for (SYMS_U64 i = 0; i < seq_count; i += 1){
+    // get sequence range & inc
+    SYMS_U64 first = *seq_index_ptr;
+    seq_index_ptr += 1;
+    SYMS_U64 opl = *seq_index_ptr;
+    
+    // iterate lines
+    SYMS_U64 last = opl - 1;
+    SYMS_Line *line_ptr = line_array + first;
+    for (SYMS_U64 j = first; j < last; j += 1, line_ptr += 1){
+      if (line_ptr->src_coord.file_id == file_id &&
+          line_ptr->src_coord.line == line){
+        SYMS_U64Range range = syms_make_u64_range(line_ptr->voff, (line_ptr + 1)->voff);
+        syms_u64_range_list_push(arena, &result, range);
+      }
+    }
+  }
+  
+  SYMS_ProfEnd();
+  return(result);
+}
+
+SYMS_API SYMS_SpatialMap1D*
+syms_group_sec_map_v(SYMS_Group *group){
+  SYMS_ProfBegin("syms_group_sec_map_v");
+  if (!group->sec_map_v_is_built){
+    SYMS_ASSERT(syms_thread_lane == 0);
+    
+    //- fill spatial map array
+    SYMS_SecInfoArray sec_array = syms_group_sec_info_array(group);
+    SYMS_SecInfo *sec_info = sec_array.sec_info;
+    SYMS_SpatialMap1DRange *ranges = syms_push_array_zero(group->arena, SYMS_SpatialMap1DRange, sec_array.count);
+    SYMS_SpatialMap1DRange *range_ptr = ranges;
+    for (SYMS_U64 i = 0; i < sec_array.count; i += 1, sec_info += 1){
+      if (sec_info->vrange.min < sec_info->vrange.max){
+        range_ptr->range = sec_info->vrange;
+        range_ptr->val = i + 1;
+        range_ptr += 1;
+      }
+    }
+    SYMS_U64 final_count = (SYMS_U64)(range_ptr - ranges);
+    syms_arena_put_back(group->arena, sizeof(SYMS_SpatialMap1DRange)*(sec_array.count - final_count));
+    
+    //- sort
+    if (!syms_spatial_map_1d_array_check_sorted(ranges, final_count)){
+      syms_spatial_map_1d_array_sort(ranges, final_count);
+    }
+    
+    //- assemble map
+    SYMS_SpatialMap1D map = {ranges, final_count};
+    
+    //- save to group
+    group->sec_map_v_is_built = syms_true;
+    group->sec_map_v = map;
+    
+    SYMS_ASSERT_PARANOID(syms_spatial_map_1d_invariants(&map));
+  }
+  
+  SYMS_SpatialMap1D *result = &group->sec_map_v;
+  
+  SYMS_ProfEnd();
+  return(result);
+}
+
+SYMS_API SYMS_SpatialMap1D*
+syms_group_sec_map_f(SYMS_Group *group){
+  SYMS_ProfBegin("syms_group_sec_map_f");
+  if (!group->sec_map_f_is_built){
+    SYMS_ASSERT(syms_thread_lane == 0);
+    
+    //- fill spatial map array
+    SYMS_SecInfoArray sec_array = syms_group_sec_info_array(group);
+    SYMS_SecInfo *sec_info = sec_array.sec_info;
+    SYMS_SpatialMap1DRange *ranges = syms_push_array_zero(group->arena, SYMS_SpatialMap1DRange, sec_array.count);
+    SYMS_SpatialMap1DRange *range_ptr = ranges;
+    for (SYMS_U64 i = 0; i < sec_array.count; i += 1, sec_info += 1){
+      if (sec_info->frange.min < sec_info->frange.max){
+        range_ptr->range = sec_info->frange;
+        range_ptr->val = i + 1;
+        range_ptr += 1;
+      }
+    }
+    SYMS_U64 final_count = (SYMS_U64)(range_ptr - ranges);
+    syms_arena_put_back(group->arena, sizeof(SYMS_SpatialMap1DRange)*(sec_array.count - final_count));
+    
+    //- sort
+    if (!syms_spatial_map_1d_array_check_sorted(ranges, final_count)){
+      syms_spatial_map_1d_array_sort(ranges, final_count);
+    }
+    
+    //- assemble map
+    SYMS_SpatialMap1D map = {ranges, final_count};
+    
+    //- save to group
+    group->sec_map_f_is_built = syms_true;
+    group->sec_map_f = map;
+    
+    SYMS_ASSERT_PARANOID(syms_spatial_map_1d_invariants(&map));
+  }
+  
+  SYMS_SpatialMap1D *result = &group->sec_map_v;
+  
+  SYMS_ProfEnd();
+  return(result);
+}
+
+SYMS_API SYMS_SpatialMap1D*
+syms_group_unit_map(SYMS_Group *group){
+  SYMS_ProfBegin("syms_group_unit_map");
+  if (!group->unit_map_is_built){
+    SYMS_ASSERT(syms_thread_lane == 0);
+    
+    //- fill spatial map array
+    SYMS_UnitRangeArray unit_ranges = syms_group_unit_ranges(group);
+    SYMS_U64 count = unit_ranges.count;
+    SYMS_UnitRange *unit_range = unit_ranges.ranges;
+    SYMS_SpatialMap1DRange *ranges = syms_push_array_zero(group->arena, SYMS_SpatialMap1DRange, count);
+    SYMS_SpatialMap1DRange *range_ptr = ranges;
+    for (SYMS_U64 i = 0; i < count; i += 1, unit_range += 1){
+      if (unit_range->vrange.min < unit_range->vrange.max){
+        range_ptr->range = unit_range->vrange;
+        range_ptr->val = unit_range->uid;
+        range_ptr += 1;
+      }
+    }
+    SYMS_U64 final_count = (SYMS_U64)(range_ptr - ranges);
+    syms_arena_put_back(group->arena, sizeof(SYMS_SpatialMap1DRange)*(count - final_count));
+    
+    //- sort
+    if (!syms_spatial_map_1d_array_check_sorted(ranges, final_count)){
+      syms_spatial_map_1d_array_sort(ranges, final_count);
+    }
+    
+    //- assemble map
+    SYMS_SpatialMap1D map = {ranges, final_count};
+    
+    //- save to group
+    group->unit_map_is_built = syms_true;
+    group->unit_map = map;
+    
+    SYMS_ASSERT_PARANOID(syms_spatial_map_1d_invariants(&map));
+  }
+  
+  SYMS_SpatialMap1D *result = &group->unit_map;
+  
+  SYMS_ProfEnd();
+  return(result);
+}
+
+SYMS_API SYMS_SpatialMap1D*
+syms_group_proc_map_from_uid(SYMS_Group *group, SYMS_UnitID uid){
+  SYMS_ProfBegin("syms_group_proc_map_from_uid");
+  SYMS_U64 index = uid - 1;
+  if (index < group->unit_count &&
+      (group->unit_cache_flags[index] & SYMS_GroupUnitCacheFlag_HasProcMap) == 0){
+    //- build loose map
+    SYMS_ArenaTemp scratch = syms_get_scratch(0, 0);
+    SYMS_SpatialMap1DLoose loose = {0};
+    
+    SYMS_String8 dbg_data = group->dbg_data;
+    SYMS_DbgAccel *dbg = group->dbg;
+    SYMS_UnitAccel *unit = syms_group_unit_from_uid(group, uid);
+    SYMS_SymbolIDArray *sid_array = syms_group_proc_sid_array_from_uid(group, uid);
+    SYMS_U64 count = sid_array->count;
+    
+    SYMS_SymbolID *sid_ptr = sid_array->ids;
+    for (SYMS_U64 i = 0; i < count; i += 1, sid_ptr += 1){
+      SYMS_SymbolID sid = *sid_ptr;
+      SYMS_SymbolKind kind = syms_group_symbol_kind_from_sid(group, unit, sid);
+      if (kind == SYMS_SymbolKind_Procedure){
+        SYMS_U64RangeArray ranges = syms_proc_vranges_from_sid(scratch.arena, dbg_data, dbg, unit, sid);
+        syms_spatial_map_1d_loose_push(scratch.arena, &loose, sid, ranges);
+      }
+    }
+    
+    //- bake tight map
+    SYMS_Arena *arena = syms_group_get_lane_arena(group);
+    SYMS_SpatialMap1D map = syms_spatial_map_1d_bake(arena, &loose);
+    
+    //- save to group
+    group->unit_cache_flags[index] |= SYMS_GroupUnitCacheFlag_HasProcMap;
+    group->unit_proc_maps[index] = map;
+    
+    syms_release_scratch(scratch);
+    
+    SYMS_ASSERT_PARANOID(syms_spatial_map_1d_invariants(&map));
+  }
+  
+  SYMS_SpatialMap1D *result = &syms_spatial_map_1d_nil;
+  if (index < group->unit_count){
+    result = &group->unit_proc_maps[index];
+  }
+  
+  SYMS_ProfEnd();
+  return(result);
+}
+
+SYMS_API SYMS_SpatialMap1D*
+syms_group_line_sequence_map_from_uid(SYMS_Group *group, SYMS_UnitID uid){
+  SYMS_ProfBegin("syms_group_line_sequence_map_from_uid");
+  SYMS_U64 index = uid - 1;
+  if (index < group->unit_count &&
+      (group->unit_cache_flags[index] & SYMS_GroupUnitCacheFlag_HasLineSeqMap) == 0){
+    //- fill spatial map array
+    SYMS_Arena *arena = syms_group_get_lane_arena(group);
+    SYMS_LineTable *line_table = syms_group_line_table_from_uid(group, uid);
+    SYMS_U64 count = line_table->sequence_count;
+    SYMS_SpatialMap1DRange *ranges = syms_push_array_zero(arena, SYMS_SpatialMap1DRange, count);
+    SYMS_SpatialMap1DRange *range_ptr = ranges;
+    SYMS_U64 *seq_index_ptr = line_table->sequence_index_array;
+    SYMS_Line *line_array = line_table->line_array;
+    for (SYMS_U64 n = 1; n <= count; n += 1, range_ptr += 1){
+      // get sequence range & inc
+      SYMS_U64 first = *seq_index_ptr;
+      seq_index_ptr += 1;
+      SYMS_U64 opl = *seq_index_ptr;
+      
+      // extract sequence range
+      SYMS_U64 last = opl - 1;
+      SYMS_U64 first_voff = line_array[first].voff;
+      SYMS_U64 opl_voff = line_array[last].voff;
+      
+      // fill range ptr
+      range_ptr->range.min = first_voff;
+      range_ptr->range.max = opl_voff;
+      range_ptr->val = n;
+    }
+    
+    //- sort
+    if (!syms_spatial_map_1d_array_check_sorted(ranges, count)){
+      syms_spatial_map_1d_array_sort(ranges, count);
+    }
+    
+    //- assemble map
+    SYMS_SpatialMap1D map = {ranges, count};
+    
+    //- save to group
+    group->unit_cache_flags[index] |= SYMS_GroupUnitCacheFlag_HasLineSeqMap;
+    group->line_sequence_maps[index] = map;
+    
+    SYMS_ASSERT_PARANOID(syms_spatial_map_1d_invariants(&map));
+  }
+  
+  SYMS_SpatialMap1D *result = &syms_spatial_map_1d_nil;
+  if (index < group->unit_count){
+    result = &group->line_sequence_maps[index];
+  }
+  
+  SYMS_ProfEnd();
+  return(result);
+}
+
+SYMS_API SYMS_SpatialMap1D*
+syms_group_var_map_from_uid(SYMS_Group *group, SYMS_UnitID uid){
+  SYMS_ProfBegin("syms_group_var_map_from_uid");
+  SYMS_ASSERT(syms_thread_lane == 0);
+  SYMS_U64 index = uid - 1;
+  if (index < group->unit_count &&
+      (group->unit_cache_flags[index] & SYMS_GroupUnitCacheFlag_HasVarMap) == 0){
+    //- build loose map
+    SYMS_ArenaTemp scratch = syms_get_scratch(0, 0);
+    SYMS_SpatialMap1DLoose loose = {0};
+    
+    SYMS_String8 dbg_data = group->dbg_data;
+    SYMS_DbgAccel *dbg = group->dbg;
+    SYMS_UnitAccel *unit = syms_group_unit_from_uid(group, uid);
+    SYMS_SymbolIDArray *sid_array = syms_group_var_sid_array_from_uid(group, uid);
+    SYMS_U64 count = sid_array->count;
+    
+    SYMS_SymbolID *sid_ptr = sid_array->ids;
+    for (SYMS_U64 n = 1; n <= count; n += 1, sid_ptr += 1){
+      SYMS_SymbolID sid = *sid_ptr;
+      SYMS_SymbolKind kind = syms_group_symbol_kind_from_sid(group, unit, sid);
+      if (kind == SYMS_SymbolKind_ImageRelativeVariable){
+        SYMS_U64 var_virt_off = syms_voff_from_var_sid(dbg_data, dbg, unit, sid);
+        SYMS_USID var_type_usid = syms_type_from_var_sid(dbg_data, dbg, unit, sid);
+        SYMS_U64 var_size = syms_group_type_size_from_usid(group, var_type_usid);
+        SYMS_U64Range range = {var_virt_off, var_virt_off + var_size};
+        syms_spatial_map_1d_loose_push_single(scratch.arena, &loose, sid, range);
+      }
+    }
+    
+    //- bake tight map
+    SYMS_Arena *arena = syms_group_get_lane_arena(group);
+    SYMS_SpatialMap1D map = syms_spatial_map_1d_bake(arena, &loose);
+    
+    //- save to group
+    group->unit_cache_flags[index] |= SYMS_GroupUnitCacheFlag_HasVarMap;
+    group->unit_var_maps[index] = map;
+    
+    syms_release_scratch(scratch);
+    
+    SYMS_ASSERT_PARANOID(syms_spatial_map_1d_invariants(&map));
+  }
+  
+  SYMS_SpatialMap1D *result = &syms_spatial_map_1d_nil;
+  if (index < group->unit_count){
+    result = &group->unit_var_maps[index];
+  }
+  
+  SYMS_ProfEnd();
+  return(result);
+}
+
+SYMS_API void
+syms_group_fetch_line_to_addr_maps_from_uid(SYMS_Group *group, SYMS_UnitID uid){
+  SYMS_ProfBegin("syms_group_fetch_line_to_addr_maps_from_uid");
+  
+  SYMS_U64 index = uid - 1;
+  if (index < group->unit_count &&
+      (group->unit_cache_flags[index] & SYMS_GroupUnitCacheFlag_HasLineToAddrMap) == 0){
+    //- setup loose map
+    SYMS_ArenaTemp scratch = syms_get_scratch(0, 0);
+    SYMS_FileToLineToAddrLoose loose = {0};
+    
+    //- last-used cache for file nodes
+    SYMS_FileToLineToAddrLooseFile *cached_file_node = 0;
+    SYMS_U64 cached_file_id = 0;
+    
+    //- read lines
+    SYMS_LineTable *line_table = syms_group_line_table_from_uid(group, uid);
+    SYMS_U64 *seq_idx_ptr = line_table->sequence_index_array;
+    SYMS_U64 seq_count = line_table->sequence_count;
+    for (SYMS_U64 i = 0; i < seq_count; i += 1){
+      // get sequence range & inc
+      SYMS_U64 first = *seq_idx_ptr;
+      seq_idx_ptr += 1;
+      SYMS_U64 opl = *seq_idx_ptr;
+      
+      // iterate lines
+      if (first < opl){
+        SYMS_Line *line_ptr = line_table->line_array + first;
+        SYMS_U64 last = opl - 1;
+        for (SYMS_U64 j = first; j < last; j += 1, line_ptr += 1){
+          // grab line data
+          SYMS_U64 line_file_id = line_ptr->src_coord.file_id;
+          SYMS_U32 line_number = line_ptr->src_coord.line;
+          
+          // get file node
+          SYMS_FileToLineToAddrLooseFile *file_node = 0;
+          if (line_file_id == cached_file_id){
+            file_node = cached_file_node;
+          }
+          else{
+            for (SYMS_FileToLineToAddrLooseFile *node = loose.first;
+                 node != 0;
+                 node = node->next){
+              if (node->file_id == line_file_id){
+                file_node = node;
+                break;
+              }
+            }
+          }
+          if (file_node == 0){
+            file_node = syms_push_array_zero(scratch.arena, SYMS_FileToLineToAddrLooseFile, 1);
+            SYMS_QueuePush(loose.first, loose.last, file_node);
+            loose.count += 1;
+            file_node->file_id = line_file_id;
+          }
+          
+          // update the file node cache slot
+          cached_file_id = line_file_id;
+          cached_file_node = file_node;
+          
+          // get line node
+          SYMS_FileToLineToAddrLooseLine *line_node = 0;
+          for (SYMS_FileToLineToAddrLooseLine *node = file_node->first;
+               node != 0;
+               node = node->next){
+            if (node->line == line_number){
+              line_node = node;
+              break;
+            }
+          }
+          if (line_node == 0){
+            line_node = syms_push_array_zero(scratch.arena, SYMS_FileToLineToAddrLooseLine, 1);
+            SYMS_QueuePush(file_node->first, file_node->last, line_node);
+            file_node->line_count += 1;
+            line_node->line = line_number;
+          }
+          
+          // push range
+          SYMS_U64Range range = syms_make_u64_range(line_ptr->voff, (line_ptr + 1)->voff);
+          syms_u64_range_list_push(scratch.arena, &line_node->ranges, range);
+          file_node->range_count += 1;
+        }
+      }
+    }
+    
+    //- convert loose to buckets & maps
+    SYMS_Arena *arena = syms_group_get_lane_arena(group);
+    
+    SYMS_FileToLineToAddrBuckets buckets = {0};
+    if (loose.count > 0){
+      buckets.bucket_count = ((loose.count + 1)*3/2) | 1;
+    }
+    buckets.buckets = syms_push_array_zero(arena, SYMS_FileToLineToAddrNode*, buckets.bucket_count);
+    for (SYMS_FileToLineToAddrLooseFile *loose_file_node = loose.first;
+         loose_file_node != 0;
+         loose_file_node = loose_file_node->next){
+      SYMS_ArenaTemp temp = syms_arena_temp_begin(scratch.arena);
+      
+      // grab counts
+      SYMS_U64 line_count = loose_file_node->line_count;
+      SYMS_U64 range_count = loose_file_node->range_count;
+      
+      // create sorted node array
+      SYMS_FileToLineToAddrLooseLine **array = syms_push_array(scratch.arena, SYMS_FileToLineToAddrLooseLine*,
+                                                               line_count);
+      {
+        SYMS_FileToLineToAddrLooseLine **line_ptr = array;
+        for (SYMS_FileToLineToAddrLooseLine *loose_line_node = loose_file_node->first;
+             loose_line_node != 0;
+             loose_line_node = loose_line_node->next, line_ptr += 1){
+          *line_ptr = loose_line_node;
+        }
+      }
+      syms_line_to_addr_line_sort(array, line_count);
+      
+      // fill line map arrays
+      SYMS_U64Range *ranges = syms_push_array(arena, SYMS_U64Range, range_count);
+      SYMS_U32 *line_range_indexes = syms_push_array(arena, SYMS_U32, line_count + 1);
+      SYMS_U32 *line_numbers = syms_push_array(arena, SYMS_U32, line_count);
+      
+      SYMS_U32 line_range_index = 0;
+      SYMS_U32 *line_range_index_ptr = line_range_indexes;
+      SYMS_U32 *line_number_ptr = line_numbers;
+      
+      {
+        SYMS_FileToLineToAddrLooseLine **line_ptr = array;
+        for (SYMS_U64 i = 0; i < line_count; i += 1, line_ptr += 1){
+          // fills
+          *line_number_ptr = (**line_ptr).line;
+          *line_range_index_ptr = line_range_index;
+          SYMS_U64Range *first_range_ptr = ranges + line_range_index;
+          SYMS_U64Range *range_ptr = first_range_ptr;
+          for (SYMS_U64RangeNode *node = (**line_ptr).ranges.first;
+               node != 0;
+               node = node->next, range_ptr += 1){
+            *range_ptr = node->range;
+          }
+          
+          // incs
+          line_number_ptr += 1;
+          line_range_index_ptr += 1;
+          line_range_index += (SYMS_U64)(range_ptr - first_range_ptr);
+        }
+        
+        // fill ender index
+        *line_range_index_ptr = line_range_index;
+      }
+      
+      // assemble the line map
+      SYMS_LineToAddrMap *new_map = syms_push_array(arena, SYMS_LineToAddrMap, 1);
+      new_map->ranges = ranges;
+      new_map->line_range_indexes = line_range_indexes;
+      new_map->line_numbers = line_numbers;
+      new_map->line_count = line_count;
+      
+      // insert bucket
+      SYMS_FileToLineToAddrNode *new_file_node = syms_push_array(arena, SYMS_FileToLineToAddrNode, 1);
+      new_file_node->file_id = loose_file_node->file_id;
+      new_file_node->map = new_map;
+      SYMS_U64 bucket_index = loose_file_node->file_id%buckets.bucket_count;
+      SYMS_StackPush(buckets.buckets[bucket_index], new_file_node);
+      
+      syms_arena_temp_end(temp);
+    }
+    
+    //- save to group
+    group->unit_cache_flags[index] |= SYMS_GroupUnitCacheFlag_HasLineToAddrMap;
+    group->file_to_line_to_addr_buckets[index] = buckets;
+    
+    syms_release_scratch(scratch);
+  }
+  
+  SYMS_ProfEnd();
+}
+
+SYMS_API SYMS_LineToAddrMap*
+syms_group_line_to_addr_map_from_uid_file_id(SYMS_Group *group, SYMS_UnitID uid, SYMS_FileID file_id){
+  SYMS_ProfBegin("syms_group_line_to_addr_map_from_uid_file_id");
+  syms_group_fetch_line_to_addr_maps_from_uid(group, uid);
+  
+  //- get line-to-addr map from file_id
+  SYMS_LineToAddrMap *result = &syms_line_to_addr_map_nil;
+  SYMS_U64 index = uid - 1;
+  if (index < group->unit_count){
+    SYMS_FileToLineToAddrBuckets *buckets = &group->file_to_line_to_addr_buckets[index];
+    if (buckets->bucket_count > 0){
+      SYMS_U64 bucket_index = file_id%buckets->bucket_count;
+      for (SYMS_FileToLineToAddrNode *node = buckets->buckets[bucket_index];
+           node != 0;
+           node = node->next){
+        if (node->file_id == file_id){
+          result = node->map;
+          break;
+        }
+      }
+    }
+  }
+  
+  SYMS_ProfEnd();
+  return(result);
+}
+
+SYMS_API SYMS_U64
+syms_group_sec_number_from_voff__accelerated(SYMS_Group *group, SYMS_U64 voff){
+  SYMS_ProfBegin("syms_group_sec_number_from_voff__accelerated");
+  SYMS_SpatialMap1D *map = syms_group_sec_map_v(group);
+  SYMS_U64 result = syms_spatial_map_1d_value_from_point(map, voff);
+  SYMS_ProfEnd();
+  return(result);
+}
+
+SYMS_API SYMS_U64
+syms_group_sec_number_from_foff__accelerated(SYMS_Group *group, SYMS_U64 foff){
+  SYMS_ProfBegin("syms_group_sec_number_from_foff__accelerated");
+  SYMS_SpatialMap1D *map = syms_group_sec_map_f(group);
+  SYMS_U64 result = syms_spatial_map_1d_value_from_point(map, foff);
+  SYMS_ProfEnd();
+  return(result);
+}
+
+SYMS_API SYMS_U64Maybe
+syms_group_sec_voff_from_foff__accelerated(SYMS_Group *group, SYMS_U64 foff){
+  SYMS_ProfBegin("syms_group_sec_voff_from_foff__accelerated");
+  SYMS_U64Maybe result = {0};
+  SYMS_U64 sec_number = syms_group_sec_number_from_foff__accelerated(group, foff);
+  SYMS_SecInfo *info = syms_group_sec_info_from_number(group, sec_number);
+  if (info != 0){
+    result.valid = syms_true;
+    result.u64 = foff - info->frange.min + info->vrange.min;
+  }
+  SYMS_ProfEnd();
+  return(result);
+}
+
+SYMS_API SYMS_U64Maybe
+syms_group_sec_foff_from_voff__accelerated(SYMS_Group *group, SYMS_U64 voff){
+  SYMS_ProfBegin("syms_group_sec_foff_from_voff__accelerated");
+  SYMS_U64Maybe result = {0};
+  SYMS_U64 sec_number = syms_group_sec_number_from_voff__accelerated(group, voff);
+  SYMS_SecInfo *info = syms_group_sec_info_from_number(group, sec_number);
+  if (info != 0){
+    result.valid = syms_true;
+    result.u64 = voff - info->vrange.min + info->frange.min;
+  }
+  SYMS_ProfEnd();
+  return(result);
+}
+
+SYMS_API SYMS_UnitID
+syms_group_uid_from_voff__accelerated(SYMS_Group *group, SYMS_U64 voff){
+  SYMS_ProfBegin("syms_group_uid_from_voff__accelerated");
+  SYMS_SpatialMap1D *map = syms_group_unit_map(group);
+  SYMS_UnitID result = syms_spatial_map_1d_value_from_point(map, voff);
+  SYMS_ProfEnd();
+  return(result);
+}
+
+SYMS_API SYMS_SymbolID
+syms_group_proc_sid_from_uid_voff__accelerated(SYMS_Group *group, SYMS_UnitID uid, SYMS_U64 voff){
+  SYMS_ProfBegin("syms_group_proc_sid_from_uid_voff__accelerated");
+  SYMS_SpatialMap1D *map = syms_group_proc_map_from_uid(group, uid);
+  SYMS_SymbolID result = syms_spatial_map_1d_value_from_point(map, voff);
+  SYMS_ProfEnd();
+  return(result);
+}
+
+SYMS_API SYMS_SymbolID
+syms_group_var_sid_from_uid_voff__accelerated(SYMS_Group *group, SYMS_UnitID uid, SYMS_U64 voff){
+  SYMS_ProfBegin("syms_group_var_sid_from_uid_voff__accelerated");
+  SYMS_SpatialMap1D *map = syms_group_var_map_from_uid(group, uid);
+  SYMS_SymbolID result = syms_spatial_map_1d_value_from_point(map, voff);
+  SYMS_ProfEnd();
+  return(result);
+}
+
+SYMS_API SYMS_Line
+syms_group_line_from_uid_voff__accelerated(SYMS_Group *group, SYMS_UnitID uid, SYMS_U64 voff){
+  SYMS_ProfBegin("syms_group_line_from_uid_voff__accelerated");
+  SYMS_SpatialMap1D *map = syms_group_line_sequence_map_from_uid(group, uid);
+  SYMS_U64 sequence_number = syms_spatial_map_1d_value_from_point(map, voff);
+  
+  SYMS_LineTable *line_table = syms_group_line_table_from_uid(group, uid);
+  SYMS_Line result = syms_line_from_sequence_voff(line_table, sequence_number, voff);
+  
+  SYMS_ProfEnd();
+  return(result);
+}
+
+SYMS_API SYMS_U64RangeArray
+syms_group_vranges_from_uid_line__accelerated(SYMS_Group *group, SYMS_UnitID uid,
+                                              SYMS_FileID file_id, SYMS_U32 line){
+  SYMS_ProfBegin("syms_group_vranges_from_uid_line__accelerated");
+  SYMS_LineToAddrMap *line_map = syms_group_line_to_addr_map_from_uid_file_id(group, uid, file_id);
+  SYMS_U64 count = line_map->line_count;
+  SYMS_U32 *numbers = line_map->line_numbers;
+  SYMS_U64 index = syms_index_from_n__u32__binary_search_round_up(numbers, count, line);
+  SYMS_U64RangeArray result = {0};
+  if (index < count && numbers[index] == line){
+    SYMS_U32 *range_indexes = line_map->line_range_indexes;
+    result.ranges = line_map->ranges + range_indexes[index];
+    result.count = range_indexes[index + 1] - range_indexes[index];
+  }
+  SYMS_ProfEnd();
+  return(result);
+}
+
+//- line info binary search helper
+
+SYMS_API SYMS_U64
+syms_index_from_n__u32__binary_search_round_up(SYMS_U32 *v, SYMS_U64 count, SYMS_U32 n){
+  SYMS_ProfBegin("syms_index_from_n__u32__binary_search_round_up");
+  SYMS_U64 result = SYMS_U64_MAX;
+  if (count > 0 && n <= v[count - 1]){
+    //- binary search:
+    //   minimum index s.t. v[index] >= n
+    //  in this one we assume:
+    //   (i != j) implies (v[i] != v[j])
+    //   thus if (v[index] == n) then index already satisfies the requirement
+    SYMS_U64 first = 0;
+    SYMS_U64 opl = count;
+    for (;;){
+      SYMS_U64 mid = (first + opl - 1)/2;
+      SYMS_U64 w = v[mid];
+      if (w < n){
+        first = mid + 1;
+      }
+      else if (w > n){
+        opl = mid + 1;
+      }
+      else{
+        first = mid;
+        break;
+      }
+      if (first + 1 >= opl){
+        break;
+      }
+    }
+    result = first;
+  }
+  SYMS_ProfEnd();
+  return(result);
+}
+
+SYMS_API void
+syms_line_to_addr_line_sort(SYMS_FileToLineToAddrLooseLine **array, SYMS_U64 count){
+  SYMS_ProfBegin("syms_line_to_addr_line_sort");
+  syms_line_to_addr_line_sort__rec(array, count);
+  SYMS_ProfEnd();
+}
+
+SYMS_API void
+syms_line_to_addr_line_sort__rec(SYMS_FileToLineToAddrLooseLine **array, SYMS_U64 count){
+  if (count > 4){
+    SYMS_U64 last = count - 1;
+    
+    // swap
+    SYMS_U64 mid = count/2;
+    SYMS_Swap(SYMS_FileToLineToAddrLooseLine*, array[mid], array[last]);
+    
+    // partition
+    SYMS_U64 key = array[last]->line;
+    SYMS_U64 j = 0;
+    for (SYMS_U64 i = 0; i < last; i += 1){
+      if (array[i]->line < key){
+        if (j != i){
+          SYMS_Swap(SYMS_FileToLineToAddrLooseLine*, array[i], array[j]);
+        }
+        j += 1;
+      }
+    }
+    
+    SYMS_Swap(SYMS_FileToLineToAddrLooseLine*, array[j], array[last]);
+    
+    // recurse
+    SYMS_U64 pivot = j;
+    syms_line_to_addr_line_sort__rec(array, pivot);
+    syms_line_to_addr_line_sort__rec(array + pivot + 1, (count - pivot - 1));
+  }
+  else if (count == 2){
+    if (array[0]->line > array[1]->line){
+      SYMS_Swap(SYMS_FileToLineToAddrLooseLine*, array[0], array[1]);
+    }
+  }
+  else if (count == 3){
+    if (array[0]->line > array[1]->line){
+      SYMS_Swap(SYMS_FileToLineToAddrLooseLine*, array[0], array[1]);
+    }
+    if (array[1]->line > array[2]->line){
+      SYMS_Swap(SYMS_FileToLineToAddrLooseLine*, array[1], array[2]);
+      if (array[0]->line > array[1]->line){
+        SYMS_Swap(SYMS_FileToLineToAddrLooseLine*, array[0], array[1]);
+      }
+    }
+  }
+  else if (count == 4){
+    if (array[0]->line > array[1]->line){
+      SYMS_Swap(SYMS_FileToLineToAddrLooseLine*, array[0], array[1]);
+    }
+    if (array[2]->line > array[3]->line){
+      SYMS_Swap(SYMS_FileToLineToAddrLooseLine*, array[2], array[3]);
+    }
+    if (array[0]->line > array[2]->line){
+      SYMS_Swap(SYMS_FileToLineToAddrLooseLine*, array[0], array[2]);
+    }
+    if (array[1]->line > array[3]->line){
+      SYMS_Swap(SYMS_FileToLineToAddrLooseLine*, array[1], array[3]);
+    }
+    if (array[1]->line > array[2]->line){
+      SYMS_Swap(SYMS_FileToLineToAddrLooseLine*, array[1], array[2]);
+    }
+  }
+}
+
+////////////////////////////////
+// NOTE(allen): Syms Group Type Graph Functions
+
+SYMS_GLOBAL SYMS_TypeNode syms_type_graph_nil = {
+  SYMS_TypeKind_Null,     // kind
+  {(SYMS_U8*)"(nil)", 5}, // name
+  0,                      // byte_size
+  0,                      // src_coord
+  &syms_type_graph_nil,   // direct_type
+  &syms_type_graph_nil,   // this_type
+};
+
+SYMS_API SYMS_TypeNode*
+syms_group_type_from_usid(SYMS_Group *group, SYMS_USID usid){
+  SYMS_ProfBegin("syms_group_type_from_usid");
+  SYMS_ASSERT(syms_thread_lane == 0);
+  SYMS_TypeNode *result = syms_group_type_from_usid__rec(group, usid);
+  SYMS_ProfEnd();
+  return(result);
+}
+
+SYMS_API SYMS_TypeNode*
+syms_group_type_from_usid__rec(SYMS_Group *group, SYMS_USID usid){
+  SYMS_TypeNode *result = 0;
+  
+  //- look at unit features
+  SYMS_UnitInfo info = syms_unit_info_from_uid(group->unit_set, usid.uid);
+  SYMS_B32 can_have_types = !!(info.features&SYMS_UnitFeature_Types);
+  if (can_have_types){
+    
+    //- get cached version of type
+    SYMS_TypeNode *type_from_cache = syms_type_from_usid(&group->type_usid_buckets, usid);
+    
+    //- construct new graph node
+    SYMS_TypeNode *new_type = 0;
+    if (type_from_cache == 0){
+      SYMS_UnitAccel *unit = syms_group_unit_from_uid(group, usid.uid);
+      SYMS_String8 dbg_data = group->dbg_data;
+      SYMS_DbgAccel *dbg = group->dbg;
+      
+      // read basic info
+      SYMS_ArenaTemp scratch = syms_get_scratch(0, 0);
+      
+      // is this a type?
+      SYMS_SymbolKind kind = syms_group_symbol_kind_from_sid(group, unit, usid.sid);
+      if (kind != SYMS_SymbolKind_Type){
+        syms_type_usid_buckets_insert(group->arena, &group->type_usid_buckets, usid, new_type);
+      }
+      else{
+        
+        // TODO(allen): group wrapper
+        SYMS_TypeInfo type_info = syms_type_info_from_sid(dbg_data, dbg, unit, usid.sid);
+        SYMS_B32 pre_inserted = syms_false;
+        
+        if (syms_type_kind_is_basic(type_info.kind)){
+          SYMS_String8 name = syms_group_symbol_name_from_sid(scratch.arena, group, unit, usid.sid);
+          new_type = syms_group_type_basic(group, type_info.kind, type_info.reported_size, name);
+        }
+        else switch (type_info.kind){
+          case SYMS_TypeKind_Bitfield:
+          case SYMS_TypeKind_Struct:
+          case SYMS_TypeKind_Class:
+          case SYMS_TypeKind_Union:
+          case SYMS_TypeKind_Enum:
+          case SYMS_TypeKind_Typedef:
+          case SYMS_TypeKind_ForwardStruct:
+          case SYMS_TypeKind_ForwardClass:
+          case SYMS_TypeKind_ForwardUnion:
+          case SYMS_TypeKind_ForwardEnum:
+          {
+            // stub in the new type (necessary to do this before we infer more types to handle cycles)
+            new_type = syms_push_array_zero(group->arena, SYMS_TypeNode, 1);
+            syms_type_usid_buckets_insert(group->arena, &group->type_usid_buckets, usid, new_type);
+            pre_inserted = syms_true;
+            
+            // resolve references
+            SYMS_TypeNode *referenced_type = &syms_type_graph_nil;
+            if (type_info.reported_size_interp == SYMS_SizeInterpretation_ResolveForwardReference){
+              SYMS_TypeKind match_kind = syms_type_kind_main_from_fwd(type_info.kind);
+              
+              SYMS_String8 name = syms_group_symbol_name_from_sid(scratch.arena, group, unit, usid.sid);
+              SYMS_MapAndUnit type_map = syms_group_type_map(group);
+              SYMS_USIDList matches = syms_symbol_list_from_string(scratch.arena, group->dbg_data, group->dbg,
+                                                                   type_map, name);
+              
+              SYMS_USID match_usid = {0};
+              for (SYMS_USIDNode *node = matches.first;
+                   node != 0;
+                   node = node->next){
+                SYMS_UnitAccel *check_unit = syms_group_unit_from_uid(group, node->usid.uid);
+                SYMS_TypeInfo check_info = syms_type_info_from_sid(group->dbg_data, group->dbg, check_unit,
+                                                                   node->usid.sid);
+                if (check_info.kind == match_kind){
+                  match_usid = node->usid;
+                  break;
+                }
+              }
+              
+              referenced_type = syms_group_type_from_usid__rec(group, match_usid);
+            }
+            
+            // resolve direct type
+            SYMS_TypeNode *direct = syms_group_type_from_usid__rec(group, type_info.direct_type);
+            
+            // read members
+            SYMS_U64 mem_count = 0;
+            SYMS_TypeMember *mems = 0;
+            if (type_info.kind == SYMS_TypeKind_Class ||
+                type_info.kind == SYMS_TypeKind_Struct ||
+                type_info.kind == SYMS_TypeKind_Union){
+              SYMS_MemsAccel *mems_accel = syms_mems_accel_from_sid(scratch.arena, dbg_data, dbg, unit,
+                                                                    usid.sid);
+              mem_count = syms_mem_count_from_mems(mems_accel);
+              mems = syms_push_array(group->arena, SYMS_TypeMember, mem_count);
+              
+              SYMS_TypeMember *mem_ptr = mems;
+              for (SYMS_U64 n = 1; n <= mem_count; n += 1, mem_ptr += 1){
+                SYMS_ArenaTemp temp = syms_arena_temp_begin(scratch.arena);
+                
+                SYMS_MemInfo mem_info = syms_mem_info_from_number(group->arena, dbg_data, dbg, unit, mems_accel, n);
+                
+                SYMS_TypeNode *type = &syms_type_graph_nil;
+                switch (mem_info.kind){
+                  case SYMS_MemKind_VTablePtr: /*no type*/ break;
+                  
+                  case SYMS_MemKind_DataField:
+                  case SYMS_MemKind_StaticData:
+                  {
+                    // TODO(allen): here we lose access to the fast path from StaticData to a virtual offset.
+                    // This path is only available from the DWARF backend, but there are a few ways we could
+                    // preserve it.
+                    SYMS_USID type_usid = syms_type_from_mem_number(dbg_data, dbg, unit, mems_accel, n);
+                    type = syms_group_type_from_usid__rec(group, type_usid);
+                  }break;
+                  
+                  case SYMS_MemKind_BaseClass:
+                  case SYMS_MemKind_VBaseClassPtr:
+                  case SYMS_MemKind_NestedType:
+                  {
+                    // directly get type from member id
+                    SYMS_USID type_usid = syms_type_from_mem_number(dbg_data, dbg, unit, mems_accel, n);
+                    type = syms_group_type_from_usid__rec(group, type_usid);
+                  }break;
+                  
+                  case SYMS_MemKind_Method:
+                  case SYMS_MemKind_StaticMethod:
+                  {
+                    // TODO(allen): here we lose access to the fast path a procedure symbol.
+                    // This path is only available from the DWARF backend, but there are a few ways we could
+                    // preserve it.
+                    SYMS_SigInfo sig = syms_sig_info_from_mem_number(scratch.arena, dbg_data, dbg, unit, mems_accel, n);
+                    
+                    SYMS_USID return_usid = syms_make_usid(sig.uid, sig.return_type_id);
+                    SYMS_USID this_usid = syms_make_usid(sig.uid, sig.this_type_id);
+                    
+                    SYMS_TypeNode *ret_type  = syms_group_type_from_usid__rec(group, return_usid);
+                    SYMS_TypeNode *this_type = syms_group_type_from_usid__rec(group, this_usid);
+                    SYMS_U64 param_count = sig.param_type_ids.count;
+                    SYMS_TypeNode **params = syms_push_array(group->arena, SYMS_TypeNode*, param_count);
+                    
+                    {
+                      SYMS_TypeNode **param_ptr = params;
+                      SYMS_TypeNode **param_opl = params + sig.param_type_ids.count;
+                      SYMS_SymbolID *param_id_ptr = sig.param_type_ids.ids;
+                      for (; param_ptr < param_opl; param_id_ptr += 1, param_ptr += 1){
+                        *param_ptr = syms_group_type_from_usid__rec(group, syms_make_usid(sig.uid, *param_id_ptr));
+                      }
+                    }
+                    
+                    type = syms_group_type_proc_from_type(group, ret_type, this_type, params, param_count);
+                  }break;
+                }
+                
+                mem_ptr->kind = mem_info.kind;
+                mem_ptr->visibility = mem_info.visibility;
+                mem_ptr->flags = mem_info.flags;
+                mem_ptr->name = syms_string_cons(group->arena, &group->string_cons, mem_info.name);
+                mem_ptr->off = mem_info.off;
+                mem_ptr->virtual_off = mem_info.virtual_off;
+                mem_ptr->type = type;
+                
+                syms_arena_temp_end(temp);
+              }
+            }
+            
+            // read constants
+            SYMS_EnumInfoArray enum_info_array = {0};
+            if (type_info.kind == SYMS_TypeKind_Enum){
+              enum_info_array = syms_enum_info_array_from_sid(group->arena, dbg_data, dbg, unit, usid.sid);
+              
+              SYMS_EnumInfo *enum_info = enum_info_array.enum_info;
+              for (SYMS_U64 i = 0; i < enum_info_array.count; i += 1, enum_info += 1){
+                enum_info->name = syms_string_cons(group->arena, &group->string_cons, enum_info->name);
+              }
+            }
+            
+            // compute size info
+            SYMS_U64 byte_size = 0;
+            switch (type_info.reported_size_interp){
+              case SYMS_SizeInterpretation_ByteCount:
+              {
+                byte_size = type_info.reported_size;
+              }break;
+              
+              case SYMS_SizeInterpretation_Multiplier:
+              {
+                SYMS_U64 next_size = 1;
+                if (direct->byte_size != 0){
+                  next_size = direct->byte_size;
+                }
+                byte_size = type_info.reported_size*next_size;
+              }break;
+              
+              case SYMS_SizeInterpretation_ResolveForwardReference:
+              {
+                byte_size = referenced_type->byte_size;
+              }break;
+            }
+            
+            // pick direct type
+            SYMS_TypeNode *direct_type = direct;
+            if (direct_type == &syms_type_graph_nil){
+              direct_type = referenced_type;
+            }
+            
+            // setup src coord
+            SYMS_TypeSrcCoord *src_coord = syms_push_array_zero(group->arena, SYMS_TypeSrcCoord, 1);
+            src_coord->usid = usid;
+            src_coord->file_id = type_info.src_coord.file_id;
+            src_coord->line = type_info.src_coord.line;
+            src_coord->col = type_info.src_coord.col;
+            
+            // get name
+            SYMS_String8 name = syms_group_symbol_name_from_sid(scratch.arena, group, unit, usid.sid);
+            
+            // fill new type
+            new_type->kind = type_info.kind;
+            new_type->name = syms_string_cons(group->arena, &group->string_cons, name);
+            new_type->byte_size = byte_size;
+            new_type->src_coord = src_coord;
+            new_type->direct_type = direct_type;
+            new_type->this_type = &syms_type_graph_nil;
+          }break;
+          
+          case SYMS_TypeKind_Modifier:
+          {
+            SYMS_TypeNode *direct = syms_group_type_from_usid__rec(group, type_info.direct_type);
+            new_type = syms_group_type_mod_from_type(group, direct, type_info.mods);
+          }break;
+          
+          case SYMS_TypeKind_Ptr:
+          case SYMS_TypeKind_LValueReference:
+          case SYMS_TypeKind_RValueReference:
+          {
+            SYMS_TypeNode *direct = syms_group_type_from_usid__rec(group, type_info.direct_type);
+            new_type = syms_group_type_ptr_from_type(group, type_info.kind, direct);
+          }break;
+          
+          case SYMS_TypeKind_Array:
+          {
+            SYMS_TypeNode *direct = syms_group_type_from_usid__rec(group, type_info.direct_type);
+            
+            SYMS_U64 array_count = type_info.reported_size;
+            if (type_info.reported_size_interp == SYMS_SizeInterpretation_ByteCount){
+              SYMS_U64 next_size = 1;
+              if (direct->byte_size != 0){
+                next_size = direct->byte_size;
+              }
+              array_count = type_info.reported_size/next_size;
+            }
+            
+            new_type = syms_group_type_array_from_type(group, direct, array_count);
+          }break;
+          
+          case SYMS_TypeKind_Proc:
+          {
+            SYMS_TypeNode *direct = syms_group_type_from_usid__rec(group, type_info.direct_type);
+            
+            // read signature
+            SYMS_SigInfo sig_info = syms_sig_info_from_type_sid(scratch.arena, dbg_data, dbg, unit, usid.sid);
+            
+            SYMS_U64 param_count = sig_info.param_type_ids.count;
+            SYMS_TypeNode **params = syms_push_array(group->arena, SYMS_TypeNode*, param_count);
+            
+            {
+              SYMS_TypeNode **param_ptr = params;
+              SYMS_TypeNode **param_opl = params + param_count;
+              SYMS_SymbolID *sid_ptr = sig_info.param_type_ids.ids;
+              for (;param_ptr < param_opl; sid_ptr += 1, param_ptr += 1){
+                SYMS_USID param_usid = syms_make_usid(sig_info.uid, *sid_ptr);
+                *param_ptr = syms_group_type_from_usid__rec(group, param_usid);
+              }
+            }
+            
+            SYMS_USID this_usid = syms_make_usid(sig_info.uid, sig_info.this_type_id);
+            SYMS_TypeNode *this_type = syms_group_type_from_usid__rec(group, this_usid);
+            
+            new_type = syms_group_type_proc_from_type(group, direct, this_type, params, param_count);
+          }break;
+          
+          case SYMS_TypeKind_MemberPtr:
+          {
+            SYMS_TypeNode *direct = syms_group_type_from_usid__rec(group, type_info.direct_type);
+            SYMS_TypeNode *container = syms_group_type_from_usid__rec(group, type_info.containing_type);
+            new_type = syms_group_type_member_ptr_from_type(group, container, direct);
+            if (type_info.reported_size_interp == SYMS_SizeInterpretation_ByteCount){
+              new_type->byte_size = type_info.reported_size;
+            }
+          }break;
+          
+          case SYMS_TypeKind_Variadic:
+          {
+            // TODO(allen): ?
+          }break;
+          
+          case SYMS_TypeKind_Label:
+          {
+            // TODO(allen): ?
+          }break;
+        }
+        
+        //- split out modifiers into stand alone nodes
+        if (type_info.mods != 0 && type_info.kind != SYMS_TypeKind_Modifier){
+          SYMS_TypeNode *direct = new_type;
+          new_type = syms_group_type_mod_from_type(group, direct, type_info.mods);
+        }
+        
+        if (!pre_inserted){
+          syms_type_usid_buckets_insert(group->arena, &group->type_usid_buckets, usid, new_type);
+        }
+      }
+      
+      syms_release_scratch(scratch);
+    }
+    
+    //- set result
+    result = type_from_cache;
+    if (result == 0){
+      result = new_type;
+    }
+  }
+  
+  if (result == 0){
+    result = &syms_type_graph_nil;
+  }
+  
+  return(result);
+}
+
+SYMS_API SYMS_TypeMemberArray*
+syms_group_members_from_type(SYMS_Group *group, SYMS_TypeNode *type){
+  SYMS_ProfBegin("syms_group_members_from_type");
+  SYMS_TypeMemberArray *result = &syms_type_member_array_nil;
+  switch (type->kind){
+    case SYMS_TypeKind_Class:
+    case SYMS_TypeKind_Struct:
+    case SYMS_TypeKind_Union:
+    {
+      if (type->lazy_ptr == 0 && type->src_coord != 0){
+        SYMS_USID usid = type->src_coord->usid;
+        SYMS_String8 dbg_data = group->dbg_data;
+        SYMS_DbgAccel *dbg = group->dbg;
+        
+        SYMS_UnitAccel *unit = syms_group_unit_from_uid(group, usid.uid);
+        
+        SYMS_ArenaTemp scratch = syms_get_scratch(0, 0);
+        SYMS_MemsAccel *mems_accel = syms_mems_accel_from_sid(scratch.arena, dbg_data, dbg, unit, usid.sid);
+        SYMS_U64 mem_count = syms_mem_count_from_mems(mems_accel);
+        SYMS_TypeMember *mems = syms_push_array(group->arena, SYMS_TypeMember, mem_count);
+        
+        SYMS_TypeMember *mem_ptr = mems;
+        for (SYMS_U64 n = 1; n <= mem_count; n += 1, mem_ptr += 1){
+          SYMS_ArenaTemp temp = syms_arena_temp_begin(scratch.arena);
+          
+          SYMS_MemInfo mem_info = syms_mem_info_from_number(group->arena, dbg_data, dbg, unit, mems_accel, n);
+          
+          SYMS_TypeNode *type = &syms_type_graph_nil;
+          switch (mem_info.kind){
+            case SYMS_MemKind_VTablePtr: /*no type*/ break;
+            
+            case SYMS_MemKind_DataField:
+            case SYMS_MemKind_StaticData:
+            {
+              // TODO(allen): here we lose access to the fast path from StaticData to a virtual offset.
+              // This path is only available from the DWARF backend, but there are a few ways we could
+              // preserve it.
+              SYMS_USID type_usid = syms_type_from_mem_number(dbg_data, dbg, unit, mems_accel, n);
+              type = syms_group_type_from_usid__rec(group, type_usid);
+            }break;
+            
+            case SYMS_MemKind_BaseClass:
+            case SYMS_MemKind_VBaseClassPtr:
+            case SYMS_MemKind_NestedType:
+            {
+              // directly get type from member id
+              SYMS_USID type_usid = syms_type_from_mem_number(dbg_data, dbg, unit, mems_accel, n);
+              type = syms_group_type_from_usid__rec(group, type_usid);
+            }break;
+            
+            case SYMS_MemKind_Method:
+            case SYMS_MemKind_StaticMethod:
+            {
+              // TODO(allen): here we lose access to the fast path a procedure symbol.
+              // This path is only available from the DWARF backend, but there are a few ways we could
+              // preserve it.
+              SYMS_SigInfo sig = syms_sig_info_from_mem_number(scratch.arena, dbg_data, dbg, unit, mems_accel, n);
+              
+              SYMS_USID return_usid = syms_make_usid(sig.uid, sig.return_type_id);
+              SYMS_USID this_usid = syms_make_usid(sig.uid, sig.this_type_id);
+              
+              SYMS_TypeNode *ret_type  = syms_group_type_from_usid__rec(group, return_usid);
+              SYMS_TypeNode *this_type = syms_group_type_from_usid__rec(group, this_usid);
+              SYMS_U64 param_count = sig.param_type_ids.count;
+              SYMS_TypeNode **params = syms_push_array(group->arena, SYMS_TypeNode*, param_count);
+              
+              {
+                SYMS_TypeNode **param_ptr = params;
+                SYMS_TypeNode **param_opl = params + sig.param_type_ids.count;
+                SYMS_SymbolID *param_id_ptr = sig.param_type_ids.ids;
+                for (; param_ptr < param_opl; param_id_ptr += 1, param_ptr += 1){
+                  *param_ptr = syms_group_type_from_usid__rec(group, syms_make_usid(sig.uid, *param_id_ptr));
+                }
+              }
+              
+              type = syms_group_type_proc_from_type(group, ret_type, this_type, params, param_count);
+            }break;
+          }
+          
+          mem_ptr->kind = mem_info.kind;
+          mem_ptr->visibility = mem_info.visibility;
+          mem_ptr->flags = mem_info.flags;
+          mem_ptr->name = syms_string_cons(group->arena, &group->string_cons, mem_info.name);
+          mem_ptr->off = mem_info.off;
+          mem_ptr->virtual_off = mem_info.virtual_off;
+          mem_ptr->type = type;
+          
+          syms_arena_temp_end(temp);
+        }
+        
+        syms_release_scratch(scratch);
+        
+        
+        SYMS_TypeMemberArray *type_member_array = syms_push_array(group->arena, SYMS_TypeMemberArray, 1);
+        type_member_array->mems = mems;
+        type_member_array->count = mem_count;
+        type->lazy_ptr = type_member_array;
+      }
+      
+      result = (SYMS_TypeMemberArray*)type->lazy_ptr;
+    }break;
+  }
+  SYMS_ProfEnd();
+  return(result);
+}
+
+SYMS_API SYMS_EnumInfoArray*
+syms_group_enum_members_from_type(SYMS_Group *group, SYMS_TypeNode *type){
+  SYMS_ProfBegin("syms_group_enum_members_from_type");
+  SYMS_EnumInfoArray *result = &syms_enum_info_array_nil;
+  if (type->kind == SYMS_TypeKind_Enum){
+    if (type->lazy_ptr == 0 && type->src_coord != 0){
+      SYMS_USID usid = type->src_coord->usid;
+      SYMS_UnitAccel *unit = syms_group_unit_from_uid(group, usid.uid);
+      SYMS_EnumInfoArray *enum_info_array = syms_push_array(group->arena, SYMS_EnumInfoArray, 1);
+      *enum_info_array = syms_enum_info_array_from_sid(group->arena, group->dbg_data, group->dbg,
+                                                       unit, usid.sid);
+      type->lazy_ptr = enum_info_array;
+    }
+    result = (SYMS_EnumInfoArray*)type->lazy_ptr;
+  }
+  SYMS_ProfEnd();
+  return(result);
+}
+
+SYMS_API SYMS_U64
+syms_group_type_size_from_usid(SYMS_Group *group, SYMS_USID usid){
+  SYMS_ProfBegin("syms_group_type_size_from_usid");
+  SYMS_ASSERT(syms_thread_lane == 0);
+  SYMS_UnitAccel *unit = syms_group_unit_from_uid(group, usid.uid);
+  SYMS_TypeInfo type_info = syms_type_info_from_sid(group->dbg_data, group->dbg, unit, usid.sid);
+  SYMS_U64 result = 0;
+  switch (type_info.reported_size_interp){
+    case SYMS_SizeInterpretation_ByteCount:
+    {
+      result = type_info.reported_size;
+    }break;
+    case SYMS_SizeInterpretation_Multiplier:
+    {
+      SYMS_U64 element_size = syms_group_type_size_from_usid(group, type_info.direct_type);
+      result = element_size*type_info.reported_size;
+    }break;
+    default:
+    {
+      SYMS_TypeNode *type_node = syms_group_type_from_usid(group, usid);
+      result = type_node->byte_size;
+    }break;
+  }
+  SYMS_ProfEnd();
+  return(result);
+}
+
+SYMS_API SYMS_TypeNode*
+syms_group_type_basic(SYMS_Group *group, SYMS_TypeKind basic_kind, SYMS_U64 size, SYMS_String8 name){
+  SYMS_ProfBegin("syms_group_type_basic");
+  SYMS_ASSERT(syms_thread_lane == 0);
+  
+  SYMS_TypeNode *result = &syms_type_graph_nil;
+  
+  SYMS_String8 name_cons = syms_string_cons(group->arena, &group->string_cons, name);
+  
+  SYMS_U64 content[3];
+  content[0] = basic_kind;
+  content[1] = size;
+  content[2] = (SYMS_U64)name_cons.str;
+  SYMS_String8 data = {(SYMS_U8*)content, sizeof(content)};
+  result = syms_type_from_content(&group->type_content_buckets, data);
+  
+  if (result == 0){
+    result = syms_push_array_zero(group->arena, SYMS_TypeNode, 1);
+    syms_type_content_buckets_insert(group->arena, &group->type_content_buckets, data, result);
+    
+    result->kind = basic_kind;
+    result->name = name_cons;
+    result->byte_size = size;
+    result->direct_type = &syms_type_graph_nil;
+    result->this_type = &syms_type_graph_nil;
+  }
+  
+  SYMS_ProfEnd();
+  return(result);
+}
+
+SYMS_API SYMS_TypeNode*
+syms_group_type_mod_from_type(SYMS_Group *group, SYMS_TypeNode *type, SYMS_TypeModifiers mods){
+  SYMS_ProfBegin("syms_group_type_mod_from_type");
+  SYMS_ASSERT(syms_thread_lane == 0);
+  
+  SYMS_TypeNode *result = &syms_type_graph_nil;
+  
+  SYMS_U64 content[3];
+  content[0] = SYMS_TypeKind_Modifier;
+  content[1] = mods;
+  content[2] = (SYMS_U64)type;
+  SYMS_String8 data = {(SYMS_U8*)content, sizeof(content)};
+  result = syms_type_from_content(&group->type_content_buckets, data);
+  
+  if (result == 0){
+    result = syms_push_array_zero(group->arena, SYMS_TypeNode, 1);
+    syms_type_content_buckets_insert(group->arena, &group->type_content_buckets, data, result);
+    
+    result->kind = SYMS_TypeKind_Modifier;
+    result->byte_size = type->byte_size;
+    result->direct_type = type;
+    result->this_type = &syms_type_graph_nil;
+    result->mods = mods;
+  }
+  
+  SYMS_ProfEnd();
+  return(result);
+}
+
+SYMS_API SYMS_TypeNode*
+syms_group_type_ptr_from_type(SYMS_Group *group, SYMS_TypeKind ptr_kind, SYMS_TypeNode *type){
+  SYMS_ProfBegin("syms_group_type_ptr_from_type");
+  SYMS_ASSERT(syms_thread_lane == 0);
+  
+  SYMS_TypeNode *result = &syms_type_graph_nil;
+  
+  SYMS_U64 content[2];
+  content[0] = ptr_kind;
+  content[1] = (SYMS_U64)type;
+  SYMS_String8 data = {(SYMS_U8*)content, sizeof(content)};
+  result = syms_type_from_content(&group->type_content_buckets, data);
+  
+  if (result == 0){
+    result = syms_push_array_zero(group->arena, SYMS_TypeNode, 1);
+    syms_type_content_buckets_insert(group->arena, &group->type_content_buckets, data, result);
+    
+    result->kind = ptr_kind;
+    result->byte_size = syms_group_address_size(group);
+    result->direct_type = type;
+    result->this_type = &syms_type_graph_nil;
+  }
+  
+  SYMS_ProfEnd();
+  return(result);
+}
+
+SYMS_API SYMS_TypeNode*
+syms_group_type_array_from_type(SYMS_Group *group, SYMS_TypeNode *type, SYMS_U64 count){
+  SYMS_ProfBegin("syms_group_type_array_from_type");
+  SYMS_ASSERT(syms_thread_lane == 0);
+  
+  SYMS_TypeNode *result = &syms_type_graph_nil;
+  
+  SYMS_U64 content[3];
+  content[0] = SYMS_TypeKind_Array;
+  content[1] = count;
+  content[2] = (SYMS_U64)type;
+  SYMS_String8 data = {(SYMS_U8*)content, sizeof(content)};
+  result = syms_type_from_content(&group->type_content_buckets, data);
+  
+  if (result == 0){
+    result = syms_push_array_zero(group->arena, SYMS_TypeNode, 1);
+    syms_type_content_buckets_insert(group->arena, &group->type_content_buckets, data, result);
+    
+    result->kind = SYMS_TypeKind_Array;
+    result->byte_size = count*type->byte_size;
+    result->direct_type = type;
+    result->this_type = &syms_type_graph_nil;
+    result->array_count = count;
+  }
+  
+  SYMS_ProfEnd();
+  return(result);
+}
+
+SYMS_API SYMS_TypeNode*
+syms_group_type_proc_from_type(SYMS_Group *group, SYMS_TypeNode *ret_type, SYMS_TypeNode *this_type,
+                               SYMS_TypeNode **param_types, SYMS_U64 param_count){
+  SYMS_ProfBegin("syms_group_type_proc_from_type");
+  SYMS_ASSERT(syms_thread_lane == 0);
+  
+  SYMS_TypeNode *result = &syms_type_graph_nil;
+  
+  SYMS_String8 params_as_string = {(SYMS_U8*)param_types, sizeof(*param_types)*param_count};
+  SYMS_String8 params_cons = syms_string_cons(group->arena, &group->string_cons, params_as_string);
+  
+  SYMS_U64 content[4];
+  content[0] = SYMS_TypeKind_Proc;
+  content[1] = (SYMS_U64)ret_type;
+  content[2] = (SYMS_U64)this_type;
+  content[3] = (SYMS_U64)params_cons.str;
+  SYMS_String8 data = {(SYMS_U8*)content, sizeof(content)};
+  result = syms_type_from_content(&group->type_content_buckets, data);
+  
+  if (result == 0){
+    result = syms_push_array_zero(group->arena, SYMS_TypeNode, 1);
+    syms_type_content_buckets_insert(group->arena, &group->type_content_buckets, data, result);
+    
+    result->kind = SYMS_TypeKind_Proc;
+    result->byte_size = 0;
+    result->direct_type = ret_type;
+    result->this_type = this_type;
+    result->proc.params = (SYMS_TypeNode**)(params_cons.str);
+    result->proc.param_count = param_count;
+  }
+  
+  SYMS_ProfEnd();
+  return(result);
+}
+
+SYMS_API SYMS_TypeNode*
+syms_group_type_member_ptr_from_type(SYMS_Group *group, SYMS_TypeNode *container, SYMS_TypeNode *type){
+  SYMS_ProfBegin("syms_group_type_member_ptr_from_type");
+  SYMS_ASSERT(syms_thread_lane == 0);
+  
+  SYMS_TypeNode *result = &syms_type_graph_nil;
+  
+  SYMS_U64 content[3];
+  content[0] = SYMS_TypeKind_MemberPtr;
+  content[1] = (SYMS_U64)container;
+  content[2] = (SYMS_U64)type;
+  SYMS_String8 data = {(SYMS_U8*)content, sizeof(content)};
+  result = syms_type_from_content(&group->type_content_buckets, data);
+  
+  if (result == 0){
+    result = syms_push_array_zero(group->arena, SYMS_TypeNode, 1);
+    syms_type_content_buckets_insert(group->arena, &group->type_content_buckets, data, result);
+    
+    SYMS_U64 byte_size = syms_group_address_size(group);
+    if (type->kind == SYMS_TypeKind_Proc){
+      byte_size *= 2;
+    }
+    
+    result->kind = SYMS_TypeKind_MemberPtr;
+    result->byte_size = byte_size;
+    result->direct_type = type;
+    result->this_type = container;
+  }
+  
+  SYMS_ProfEnd();
+  return(result);
+}
+
+SYMS_API SYMS_String8
+syms_group_string_from_type(SYMS_Arena *arena, SYMS_TypeNode *type){
+  SYMS_ProfBegin("syms_group_string_from_type");
+  SYMS_ArenaTemp scratch = syms_get_scratch(&arena, 1);
+  SYMS_String8List list = {0};
+  syms_group_lhs_string_from_type(scratch.arena, type, &list);
+  syms_group_rhs_string_from_type(scratch.arena, type, &list);
+  SYMS_String8 result = syms_string_list_join(arena, &list, 0);
+  syms_release_scratch(scratch);
+  SYMS_ProfEnd();
+  return(result);
+}
+
+SYMS_API void
+syms_group_lhs_string_from_type__internal(SYMS_Arena *arena, SYMS_TypeNode *type, SYMS_String8List *out,
+                                          SYMS_U32 precedence_level, SYMS_B32 skip_return){
+  switch (type->kind){
+    default:
+    {
+      syms_string_list_push(arena, out, type->name);
+      syms_string_list_push(arena, out, syms_str8_lit(" "));
+    }break;
+    
+    case SYMS_TypeKind_Bitfield:
+    {
+      syms_group_lhs_string_from_type__internal(arena, type->direct_type, out, precedence_level, skip_return);
+    }break;
+    
+    case SYMS_TypeKind_Modifier:
+    {
+      syms_group_lhs_string_from_type__internal(arena, type->direct_type, out, 1, skip_return);
+      if (type->mods & SYMS_TypeModifier_Const){
+        syms_string_list_push(arena, out, syms_str8_lit("const "));
+      }
+      if (type->mods & SYMS_TypeModifier_Restrict){
+        syms_string_list_push(arena, out, syms_str8_lit("restrict "));
+      }
+      if (type->mods & SYMS_TypeModifier_Volatile){
+        syms_string_list_push(arena, out, syms_str8_lit("volatile "));
+      }
+    }break;
+    
+    case SYMS_TypeKind_Variadic:
+    {
+      syms_string_list_push(arena, out, syms_str8_lit("..."));
+    }break;
+    
+    case SYMS_TypeKind_Label:
+    {
+      // TODO(allen): ???
+    }break;
+    
+    case SYMS_TypeKind_Struct:
+    case SYMS_TypeKind_Union:
+    case SYMS_TypeKind_Enum:
+    case SYMS_TypeKind_Class:
+    case SYMS_TypeKind_Typedef:
+    {
+      syms_string_list_push(arena, out, type->name);
+      syms_string_list_push(arena, out, syms_str8_lit(" "));
+    }break;
+    
+    case SYMS_TypeKind_ForwardStruct:
+    case SYMS_TypeKind_ForwardUnion:
+    case SYMS_TypeKind_ForwardEnum:
+    case SYMS_TypeKind_ForwardClass:
+    {
+      SYMS_String8 keyword = syms_string_from_enum_value(SYMS_TypeKind, type->kind);
+      syms_string_list_push(arena, out, keyword);
+      syms_string_list_push(arena, out, syms_str8_lit(" "));
+      syms_string_list_push(arena, out, type->name);
+      syms_string_list_push(arena, out, syms_str8_lit(" "));
+    }break;
+    
+    case SYMS_TypeKind_Array:
+    {
+      syms_group_lhs_string_from_type__internal(arena, type->direct_type, out, 2, skip_return);
+      if (precedence_level == 1){
+        syms_string_list_push(arena, out, syms_str8_lit("("));
+      }
+    }break;
+    
+    case SYMS_TypeKind_Proc:
+    {
+      if (!skip_return){
+        syms_group_lhs_string_from_type__internal(arena, type->direct_type, out, 2, syms_false);
+      }
+      if (precedence_level == 1){
+        syms_string_list_push(arena, out, syms_str8_lit("("));
+      }
+    }break;
+    
+    case SYMS_TypeKind_Ptr:
+    {
+      syms_group_lhs_string_from_type__internal(arena, type->direct_type, out, 1, skip_return);
+      syms_string_list_push(arena, out, syms_str8_lit("*"));
+    }break;
+    
+    case SYMS_TypeKind_LValueReference:
+    {
+      syms_group_lhs_string_from_type__internal(arena, type->direct_type, out, 1, skip_return);
+      syms_string_list_push(arena, out, syms_str8_lit("&"));
+    }break;
+    
+    case SYMS_TypeKind_RValueReference:
+    {
+      syms_group_lhs_string_from_type__internal(arena, type->direct_type, out, 1, skip_return);
+      syms_string_list_push(arena, out, syms_str8_lit("&&"));
+    }break;
+    
+    case SYMS_TypeKind_MemberPtr:
+    {
+      syms_group_lhs_string_from_type__internal(arena, type->direct_type, out, 1, skip_return);
+      SYMS_TypeNode *container = type->this_type;
+      if (container != 0){
+        syms_string_list_push(arena, out, container->name);
+      }
+      else{
+        syms_string_list_push(arena, out, syms_str8_lit("<unknown-class>"));
+      }
+      syms_string_list_push(arena, out, syms_str8_lit("::*"));
+    }break;
+  }
+}
+
+SYMS_API void
+syms_group_rhs_string_from_type__internal(SYMS_Arena *arena, SYMS_TypeNode *type, SYMS_String8List *out,
+                                          SYMS_U32 precedence_level){
+  switch (type->kind){
+    default:break;
+    
+    case SYMS_TypeKind_Bitfield:
+    {
+      syms_group_rhs_string_from_type__internal(arena, type->direct_type, out, precedence_level);
+    }break;
+    
+    case SYMS_TypeKind_Modifier:
+    case SYMS_TypeKind_Ptr:
+    case SYMS_TypeKind_LValueReference:
+    case SYMS_TypeKind_RValueReference:
+    case SYMS_TypeKind_MemberPtr:
+    {
+      syms_group_rhs_string_from_type__internal(arena, type->direct_type, out, 1);
+    }break;
+    
+    case SYMS_TypeKind_Array:
+    {
+      if (precedence_level == 1){
+        syms_string_list_push(arena, out, syms_str8_lit(")"));
+      }
+      
+      SYMS_String8 count_str = syms_string_from_u64(arena, type->array_count);
+      syms_string_list_push(arena, out, syms_str8_lit("["));
+      syms_string_list_push(arena, out, count_str);
+      syms_string_list_push(arena, out, syms_str8_lit("]"));
+      
+      syms_group_rhs_string_from_type__internal(arena, type->direct_type, out, 2);
+    }break;
+    
+    case SYMS_TypeKind_Proc:
+    {
+      if (precedence_level == 1){
+        syms_string_list_push(arena, out, syms_str8_lit(")"));
+      }
+      
+      // parameters
+      if (type->proc.param_count == 0){
+        syms_string_list_push(arena, out, syms_str8_lit("(void)"));
+      }
+      else{
+        syms_string_list_push(arena, out, syms_str8_lit("("));
+        SYMS_U64 param_count = type->proc.param_count;
+        SYMS_TypeNode **param = type->proc.params;
+        for (SYMS_U64 i = 0; i < param_count; i += 1, param += 1){
+          SYMS_String8 param_str = syms_group_string_from_type(arena, *param);
+          SYMS_String8 param_str_trimmed = syms_str8_skip_chop_whitespace(param_str);
+          syms_string_list_push(arena, out, param_str_trimmed);
+          if (i + 1 < param_count){
+            syms_string_list_push(arena, out, syms_str8_lit(", "));
+          }
+        }
+        syms_string_list_push(arena, out, syms_str8_lit(")"));
+      }
+      
+      syms_group_rhs_string_from_type__internal(arena, type->direct_type, out, 2);
+    }break;
+    
+    case SYMS_TypeKind_Label:
+    {
+      // TODO(allen): ???
+    }break;
+  }
+}
+
+SYMS_API void
+syms_group_lhs_string_from_type(SYMS_Arena *arena, SYMS_TypeNode *type, SYMS_String8List *out){
+  SYMS_ProfBegin("syms_group_lhs_string_from_type");
+  syms_group_lhs_string_from_type__internal(arena, type, out, 0, syms_false);
+  SYMS_ProfEnd();
+}
+
+SYMS_API void
+syms_group_rhs_string_from_type(SYMS_Arena *arena, SYMS_TypeNode *type, SYMS_String8List *out){
+  SYMS_ProfBegin("syms_group_rhs_string_from_type");
+  syms_group_rhs_string_from_type__internal(arena, type, out, 0);
+  SYMS_ProfEnd();
+}
+
+SYMS_API void
+syms_group_lhs_string_from_type_skip_return(SYMS_Arena *arena, SYMS_TypeNode *type, SYMS_String8List *out){
+  SYMS_ProfBegin("syms_group_lhs_string_from_type_skip_return");
+  syms_group_lhs_string_from_type__internal(arena, type, out, 0, syms_true);
+  SYMS_ProfEnd();
+}
+
+SYMS_API SYMS_U64
+syms_type_usid_hash(SYMS_USID usid){
+  SYMS_U64 result = syms_hash_u64(usid.sid + usid.uid*97);
+  return(result);
+}
+
+SYMS_API SYMS_TypeNode*
+syms_type_from_usid(SYMS_TypeUSIDBuckets *buckets, SYMS_USID usid){
+  SYMS_ProfBegin("syms_type_from_usid");
+  SYMS_TypeNode *result = 0;
+  SYMS_U64 hash = syms_type_usid_hash(usid);
+  SYMS_U64 bucket_index = hash%buckets->bucket_count;
+  for (SYMS_TypeUSIDNode *node = buckets->buckets[bucket_index];
+       node != 0;
+       node = node->next){
+    SYMS_USID key = node->key;
+    if (key.uid == usid.uid && key.sid == usid.sid){
+      result = node->type;
+      break;
+    }
+  }
+  SYMS_ProfEnd();
+  return(result);
+}
+
+SYMS_API void
+syms_type_usid_buckets_insert(SYMS_Arena *arena, SYMS_TypeUSIDBuckets *buckets, SYMS_USID key,
+                              SYMS_TypeNode *type){
+  SYMS_ProfBegin("syms_type_usid_buckets_insert");
+  SYMS_U64 hash = syms_type_usid_hash(key);
+  SYMS_U64 bucket_index = hash%buckets->bucket_count;
+  
+  SYMS_TypeUSIDNode *new_node = syms_push_array(arena, SYMS_TypeUSIDNode, 1);
+  SYMS_StackPush(buckets->buckets[bucket_index], new_node);
+  new_node->key = key;
+  new_node->type = type;
+  
+  SYMS_ProfEnd();
+}
+
+SYMS_API SYMS_U64
+syms_type_content_hash(SYMS_String8 data){
+  SYMS_U64 result = syms_hash_djb2(data);
+  return(result);
+}
+
+SYMS_API SYMS_TypeNode*
+syms_type_from_content(SYMS_TypeContentBuckets *buckets, SYMS_String8 data){
+  SYMS_ProfBegin("syms_type_from_content");
+  SYMS_TypeNode *result = 0;
+  SYMS_U64 hash = syms_type_content_hash(data);
+  SYMS_U64 bucket_index = hash%buckets->bucket_count;
+  for (SYMS_TypeContentNode *node = buckets->buckets[bucket_index];
+       node != 0;
+       node = node->next){
+    if (node->hash == hash && syms_string_match(node->key, data, 0)){
+      result = node->type;
+      break;
+    }
+  }
+  SYMS_ProfEnd();
+  return(result);
+}
+
+SYMS_API SYMS_String8
+syms_type_content_buckets_insert(SYMS_Arena *arena, SYMS_TypeContentBuckets *buckets,
+                                 SYMS_String8 key, SYMS_TypeNode *type){
+  SYMS_ProfBegin("syms_type_content_buckets_insert");
+  SYMS_U64 hash = syms_type_content_hash(key);
+  SYMS_U64 bucket_index = hash%buckets->bucket_count;
+  
+  SYMS_String8 result = syms_push_string_copy(arena, key);
+  
+  SYMS_TypeContentNode *new_node = syms_push_array(arena, SYMS_TypeContentNode, 1);
+  SYMS_StackPush(buckets->buckets[bucket_index], new_node);
+  new_node->key = result;
+  new_node->hash = hash;
+  new_node->type = type;
+  
+  SYMS_ProfEnd();
+  return(result);
+}
+
+////////////////////////////////
+//~ NOTE(allen): File Mapping
+
+SYMS_API SYMS_Name2FileIDMap*
+syms_group_file_map(SYMS_Group *group){
+  SYMS_ProfBegin("syms_group_file_map");
+  if (!group->name_2_file_id_map_is_built){
+    SYMS_ASSERT(syms_thread_lane == 0);
+    SYMS_ArenaTemp scratch = syms_get_scratch(0, 0);
+    
+    //- build loose file map
+    SYMS_U64 unit_count = group->unit_count;
+    SYMS_Name2FileIDMapLoose loose = {0};
+    for (SYMS_UnitID uid = 1; uid <= unit_count; uid += 1){
+      SYMS_LineParseOut *line_parse = syms_group_line_parse_from_uid(group, uid);
+      
+      // build from file_id_array
+      if (line_parse->file_id_array.count != 0){
+        SYMS_String8Array *inferred = syms_group_inferred_file_table_from_uid(group, uid);
+        SYMS_U64 count = SYMS_MIN(line_parse->file_id_array.count, inferred->count);
+        SYMS_FileID *file_id_ptr = line_parse->file_id_array.ids;
+        SYMS_FileID *opl = file_id_ptr + count;
+        SYMS_String8 *name_ptr = inferred->strings;
+        for (;file_id_ptr < opl; file_id_ptr += 1, name_ptr += 1){
+          SYMS_String8 name_cons = syms_string_cons(group->arena, &group->string_cons, *name_ptr);
+          syms_name_2_file_id_map_loose_push(scratch.arena, &loose, name_cons, uid, *file_id_ptr);
+        }
+      }
+      
+      // build from file_table
+      else{
+        SYMS_String8Array *file_table = syms_group_file_table_from_uid_with_fallbacks(group, uid);
+        SYMS_FileID file_id = 1;
+        for (SYMS_String8 *name_ptr = file_table->strings, *opl = file_table->strings + file_table->count;
+             name_ptr < opl;
+             name_ptr += 1, file_id += 1){
+          SYMS_String8 name_cons = syms_string_cons(group->arena, &group->string_cons, *name_ptr);
+          syms_name_2_file_id_map_loose_push(scratch.arena, &loose, name_cons, uid, file_id);
+        }
+      }
+    }
+    
+    //- bake
+    SYMS_Name2FileIDMap name_2_file_id_map = syms_name_2_file_id_map_bake(group->arena, &loose);
+    
+    //- save to group
+    group->name_2_file_id_map_is_built = syms_true;
+    group->name_2_file_id_map = name_2_file_id_map;
+    
+    syms_release_scratch(scratch);
+  }
+  
+  SYMS_Name2FileIDMap *result = &group->name_2_file_id_map;
+  
+  SYMS_ProfEnd();
+  return(result);
+}
+
+#endif //SYMS_GROUP_C
