Update syms library to latest version, includes dwarf5 support

diff --git a/Engine/Source/Runtime/SymsLib/syms/core/group/syms_group.h b/Engine/Source/Runtime/SymsLib/syms/core/group/syms_group.h
new file mode 100644
--- /dev/null
+++ b/Engine/Source/Runtime/SymsLib/syms/core/group/syms_group.h
@@ -0,0 +1,413 @@
+// Copyright Epic Games, Inc. All Rights Reserved.
+/* date = August 2nd 2021 0:48 pm */
+
+#ifndef SYMS_GROUP_H
+#define SYMS_GROUP_H
+
+////////////////////////////////
+// NOTE(allen): Syms Type Graph
+
+typedef struct SYMS_TypeMember{
+  SYMS_MemKind kind;
+  SYMS_MemVisibility visibility;
+  SYMS_MemFlags flags;
+  SYMS_String8 name;
+  SYMS_U32 off;
+  SYMS_U32 virtual_off;
+  struct SYMS_TypeNode *type;
+} SYMS_TypeMember;
+
+typedef struct SYMS_TypeSrcCoord{
+  SYMS_USID usid;
+  SYMS_FileID file_id;
+  SYMS_U32 line;
+  SYMS_U32 col;
+} SYMS_TypeSrcCoord;
+
+typedef struct SYMS_TypeMemberArray{
+  SYMS_TypeMember *mems;
+  SYMS_U64 count;
+} SYMS_TypeMemberArray;
+
+typedef struct SYMS_TypeNode{
+  // SYMS_TypeNode extends and completes the information from SYMS_TypeInfo.
+  // See SYMS_TypeInfo for more interpretation info.
+  
+  SYMS_TypeKind kind;
+  SYMS_String8 name;
+  SYMS_U64 byte_size;
+  
+  // when non-null contains the source location of the type's definition
+  SYMS_TypeSrcCoord *src_coord;
+  
+  // (in addition to interpretations of SYMS_TypeInfo 'direct_type')
+  //  SYMS_TypeKind_Forward*   -> the concrete type referenced by the forward reference
+  struct SYMS_TypeNode *direct_type;
+  
+  // 'this_type' meaning depends on kind:
+  //  SYMS_TypeKind_MemberPtr  -> the container type of the member pointer
+  //  SYMS_TypeKind_Proc       -> if non-nil this is the type of an implicit 'this' in a C++ method
+  struct SYMS_TypeNode *this_type;
+  
+  union{
+    // kind: SYMS_TypeKind_Modifier
+    SYMS_TypeModifiers mods;
+    
+    // kind: SYMS_TypeKind_Array
+    SYMS_U64 array_count;
+    
+    // kind: SYMS_TypeKind_Proc
+    struct{
+      struct SYMS_TypeNode **params;
+      SYMS_U64 param_count;
+    } proc;
+    
+    // opaque pointer for lazy eval attachments to the type node.
+    void *lazy_ptr;
+  };
+} SYMS_TypeNode;
+
+typedef struct SYMS_TypeUSIDNode{
+  struct SYMS_TypeUSIDNode *next;
+  SYMS_USID key;
+  SYMS_TypeNode *type;
+} SYMS_TypeUSIDNode;
+
+typedef struct SYMS_TypeUSIDBuckets{
+  SYMS_TypeUSIDNode **buckets;
+  SYMS_U64 bucket_count;
+} SYMS_TypeUSIDBuckets;
+
+typedef struct SYMS_TypeContentNode{
+  struct SYMS_TypeContentNode *next;
+  SYMS_String8 key;
+  SYMS_U64 hash;
+  SYMS_TypeNode *type;
+} SYMS_TypeContentNode;
+
+typedef struct SYMS_TypeContentBuckets{
+  SYMS_TypeContentNode **buckets;
+  SYMS_U64 bucket_count;
+} SYMS_TypeContentBuckets;
+
+////////////////////////////////
+// NOTE(allen): Syms Line Mapping Structures
+
+typedef struct SYMS_LineToAddrMap{
+  SYMS_U64Range *ranges;
+  // line_range_indexes ranges from [0,line_count] inclusive so that:
+  // for-all i in [0,line_count):
+  //   (line_range_indexes[i + 1] - line_range_indexes[i]) == # of ranges for line at index i
+  SYMS_U32 *line_range_indexes;
+  SYMS_U32 *line_numbers;
+  SYMS_U64 line_count;
+} SYMS_LineToAddrMap;
+
+typedef struct SYMS_FileToLineToAddrNode{
+  struct SYMS_FileToLineToAddrNode *next;
+  SYMS_FileID file_id;
+  SYMS_LineToAddrMap *map;
+} SYMS_FileToLineToAddrNode;
+
+typedef struct SYMS_FileToLineToAddrBuckets{
+  SYMS_FileToLineToAddrNode **buckets;
+  SYMS_U64 bucket_count;
+} SYMS_FileToLineToAddrBuckets;
+
+//- loose version
+typedef struct SYMS_FileToLineToAddrLooseLine{
+  struct SYMS_FileToLineToAddrLooseLine *next;
+  SYMS_U32 line;
+  SYMS_U64RangeList ranges;
+} SYMS_FileToLineToAddrLooseLine;
+
+typedef struct SYMS_FileToLineToAddrLooseFile{
+  struct SYMS_FileToLineToAddrLooseFile *next;
+  SYMS_FileID file_id;
+  SYMS_FileToLineToAddrLooseLine *first;
+  SYMS_FileToLineToAddrLooseLine *last;
+  SYMS_U64 line_count;
+  SYMS_U64 range_count;
+} SYMS_FileToLineToAddrLooseFile;
+
+typedef struct SYMS_FileToLineToAddrLoose{
+  SYMS_FileToLineToAddrLooseFile *first;
+  SYMS_FileToLineToAddrLooseFile *last;
+  SYMS_U64 count;
+} SYMS_FileToLineToAddrLoose;
+
+////////////////////////////////
+// NOTE(allen): Syms Group Types
+
+typedef SYMS_U32 SYMS_GroupUnitCacheFlags;
+enum{
+  SYMS_GroupUnitCacheFlag_HasProcSidArray  = (1 << 0),
+  SYMS_GroupUnitCacheFlag_HasVarSidArray   = (1 << 1),
+  SYMS_GroupUnitCacheFlag_HasTypeSidArray  = (1 << 2),
+  SYMS_GroupUnitCacheFlag_HasFileTable     = (1 << 3),
+  SYMS_GroupUnitCacheFlag_HasInfFileTable  = (1 << 4),
+  SYMS_GroupUnitCacheFlag_HasLineTable     = (1 << 5),
+  SYMS_GroupUnitCacheFlag_HasProcMap       = (1 << 6),
+  SYMS_GroupUnitCacheFlag_HasVarMap        = (1 << 7),
+  SYMS_GroupUnitCacheFlag_HasLineSeqMap    = (1 << 8),
+  SYMS_GroupUnitCacheFlag_HasLineToAddrMap = (1 << 9),
+};
+
+typedef struct SYMS_Group{
+  SYMS_Arena *arena;
+  
+  //- thread lanes
+  SYMS_Arena **lane_arenas;
+  SYMS_U64 lane_count;
+  SYMS_U64 lane_max;
+  
+  //- data for binary and debug files
+  SYMS_String8 bin_data;
+  SYMS_String8 dbg_data;
+  SYMS_BinAccel *bin;
+  SYMS_DbgAccel *dbg;
+  
+  //- top-level accelerators and info
+  SYMS_Arch arch;
+  SYMS_U64 address_size;
+  SYMS_U64 default_vbase;
+  SYMS_SecInfoArray sec_info_array;
+  SYMS_UnitSetAccel *unit_set;
+  SYMS_U64 unit_count;
+  SYMS_MapAccel *type_map;
+  
+  //- basic section caches
+  SYMS_String8 *sec_names;
+  
+  //- basic unit caches
+  SYMS_GroupUnitCacheFlags *unit_cache_flags;
+  SYMS_UnitAccel **units;
+  SYMS_SymbolIDArray *proc_sid_arrays;
+  SYMS_SymbolIDArray *var_sid_arrays;
+  SYMS_SymbolIDArray *type_sid_arrays;
+  SYMS_String8Array *file_tables;
+  SYMS_String8Array *inferred_file_tables;
+  SYMS_LineParseOut *line_tables;
+  SYMS_SpatialMap1D *unit_proc_maps;
+  SYMS_SpatialMap1D *unit_var_maps;
+  SYMS_SpatialMap1D *line_sequence_maps;
+  SYMS_FileToLineToAddrBuckets *file_to_line_to_addr_buckets;
+  
+  //- hash tables caches
+  SYMS_StringCons string_cons;
+  SYMS_FileID2NameMap file_id_2_name_map;
+  SYMS_TypeUSIDBuckets type_usid_buckets;
+  SYMS_TypeContentBuckets type_content_buckets;
+  
+  //- one-time fills/builds
+  SYMS_B32 unit_ranges_is_filled;
+  SYMS_B32 type_map_unit_is_filled;
+  SYMS_B32 sec_map_v_is_built;
+  SYMS_B32 sec_map_f_is_built;
+  SYMS_B32 unit_map_is_built;
+  SYMS_B32 name_2_file_id_map_is_built;
+  
+  SYMS_UnitRangeArray unit_ranges;
+  SYMS_UnitAccel *type_map_unit;
+  SYMS_SpatialMap1D sec_map_v;
+  SYMS_SpatialMap1D sec_map_f;
+  SYMS_SpatialMap1D unit_map;
+  SYMS_Name2FileIDMap name_2_file_id_map;
+} SYMS_Group;
+
+////////////////////////////////
+// NOTE(allen): Data Structure Nils
+
+SYMS_GLOBAL SYMS_SymbolIDArray syms_sid_array_nil = {0};
+SYMS_GLOBAL SYMS_String8Array syms_string_array_nil = {0};
+SYMS_GLOBAL SYMS_LineParseOut syms_line_parse_nil = {0};
+SYMS_GLOBAL SYMS_SpatialMap1D syms_spatial_map_1d_nil = {0};
+SYMS_GLOBAL SYMS_LineToAddrMap syms_line_to_addr_map_nil = {0};
+SYMS_GLOBAL SYMS_TypeMemberArray syms_type_member_array_nil = {0};
+SYMS_GLOBAL SYMS_EnumInfoArray syms_enum_info_array_nil = {0};
+
+
+////////////////////////////////
+// NOTE(allen): File Inference Helper Macros
+
+// TODO(allen): Just here to help people with transition. Eventually we're
+// going to cut the helper macros just have syms_file_inf_infer*
+#define SYMS_GroupInitParams            SYMS_ParseBundle
+#define syms_group_infer_from_file_list syms_file_inf_infer_from_file_list
+#define syms_group_infer_from_file      syms_file_inf_infer_from_file
+
+////////////////////////////////
+// NOTE(allen): Syms Group Setup Functions
+
+SYMS_API SYMS_Group* syms_group_alloc(void);
+SYMS_API void        syms_group_release(SYMS_Group *group);
+
+SYMS_API void        syms_group_init(SYMS_Group *group, SYMS_ParseBundle *params);
+
+SYMS_API void        syms_group_parse_all_units__single_thread(SYMS_Group *group);
+SYMS_API void        syms_group_parse_all_top_level(SYMS_Group *group);
+
+SYMS_API void        syms_group_begin_multilane(SYMS_Group *group, SYMS_U64 lane_count);
+SYMS_API void        syms_group_end_multilane(SYMS_Group *group);
+
+SYMS_API SYMS_Arena* syms_group_get_lane_arena(SYMS_Group *group);
+
+////////////////////////////////
+// NOTE(allen): Syms Group Getters & Cache Accessors
+
+// TODO(allen): better sorting plan here.
+
+SYMS_API SYMS_String8       syms_group_bin_data(SYMS_Group *group);
+SYMS_API SYMS_BinAccel*     syms_group_bin(SYMS_Group *group);
+SYMS_API SYMS_String8       syms_group_dbg_data(SYMS_Group *group);
+SYMS_API SYMS_DbgAccel*     syms_group_dbg(SYMS_Group *group);
+SYMS_API SYMS_UnitSetAccel* syms_group_unit_set(SYMS_Group *group);
+
+SYMS_API SYMS_U64           syms_group_address_size(SYMS_Group *group);
+SYMS_API SYMS_U64           syms_group_default_vbase(SYMS_Group *group);
+
+SYMS_API SYMS_SecInfoArray  syms_group_sec_info_array(SYMS_Group *group);
+SYMS_API SYMS_SecInfo*      syms_group_sec_info_from_number(SYMS_Group *group, SYMS_U64 n);
+
+SYMS_API SYMS_U64           syms_group_unit_count(SYMS_Group *group);
+SYMS_API SYMS_UnitInfo      syms_group_unit_info_from_uid(SYMS_Group *group, SYMS_UnitID uid);
+SYMS_API SYMS_UnitNames     syms_group_unit_names_from_uid(SYMS_Arena *arena, SYMS_Group *group, SYMS_UnitID uid);
+
+// thread safe (with lanes equipped to group)
+SYMS_API SYMS_UnitAccel*    syms_group_unit_from_uid(SYMS_Group *group, SYMS_UnitID uid);
+SYMS_API SYMS_SymbolIDArray*syms_group_proc_sid_array_from_uid(SYMS_Group *group, SYMS_UnitID uid);
+SYMS_API SYMS_SymbolIDArray*syms_group_var_sid_array_from_uid(SYMS_Group *group, SYMS_UnitID uid);
+SYMS_API SYMS_SymbolIDArray*syms_group_type_sid_array_from_uid(SYMS_Group *group, SYMS_UnitID uid);
+SYMS_API SYMS_String8Array* syms_group_file_table_from_uid(SYMS_Group *group, SYMS_UnitID uid);
+SYMS_API SYMS_LineParseOut* syms_group_line_parse_from_uid(SYMS_Group *group, SYMS_UnitID uid);
+SYMS_API SYMS_LineTable*    syms_group_line_table_from_uid(SYMS_Group *group, SYMS_UnitID uid);
+SYMS_API SYMS_String8Array* syms_group_inferred_file_table_from_uid(SYMS_Group *group, SYMS_UnitID uid);
+
+SYMS_API SYMS_String8Array* syms_group_file_table_from_uid_with_fallbacks(SYMS_Group *group, SYMS_UnitID uid);
+
+SYMS_API SYMS_UnitRangeArray syms_group_unit_ranges(SYMS_Group *group);
+
+SYMS_API SYMS_SymbolKind    syms_group_symbol_kind_from_sid(SYMS_Group *group,
+                                                            SYMS_UnitAccel *unit, SYMS_SymbolID sid);
+SYMS_API SYMS_String8       syms_group_symbol_name_from_sid(SYMS_Arena *arena, SYMS_Group *group,
+                                                            SYMS_UnitAccel *unit, SYMS_SymbolID sid);
+
+SYMS_API SYMS_String8       syms_group_file_name_from_id(SYMS_Group *group, SYMS_UnitID uid, SYMS_FileID file_id);
+
+SYMS_API SYMS_MapAndUnit    syms_group_type_map(SYMS_Group *group);
+
+
+////////////////////////////////
+// NOTE(allen): Syms Group Address Mapping Functions
+
+//- linear scan versions
+SYMS_API SYMS_U64      syms_group_sec_number_from_voff__linear_scan(SYMS_Group *group, SYMS_U64 voff);
+SYMS_API SYMS_U64      syms_group_sec_number_from_foff__linear_scan(SYMS_Group *group, SYMS_U64 foff);
+
+SYMS_API SYMS_U64Maybe syms_group_voff_from_foff__linear_scan(SYMS_Group *group, SYMS_U64 foff);
+SYMS_API SYMS_U64Maybe syms_group_foff_from_voff__linear_scan(SYMS_Group *group, SYMS_U64 voff);
+
+SYMS_API SYMS_UnitID   syms_group_uid_from_voff__linear_scan(SYMS_Group *group, SYMS_U64 voff);
+SYMS_API SYMS_SymbolID syms_group_proc_sid_from_uid_voff__linear_scan(SYMS_Group *group, SYMS_UnitID uid,
+                                                                      SYMS_U64 voff);
+SYMS_API SYMS_SymbolID syms_group_var_sid_from_uid_voff__linear_scan(SYMS_Group *group, SYMS_UnitID uid,
+                                                                     SYMS_U64 voff);
+
+SYMS_API SYMS_Line     syms_group_line_from_uid_voff__linear_scan(SYMS_Group *group, SYMS_UnitID uid,
+                                                                  SYMS_U64 voff);
+SYMS_API SYMS_U64RangeList syms_group_vranges_from_uid_line__linear_scan(SYMS_Arena *arena, SYMS_Group *group,
+                                                                         SYMS_UnitID uid,
+                                                                         SYMS_FileID file_id, SYMS_U32 line);
+
+//- map getters
+SYMS_API SYMS_SpatialMap1D* syms_group_sec_map_v(SYMS_Group *group);
+SYMS_API SYMS_SpatialMap1D* syms_group_sec_map_f(SYMS_Group *group);
+
+SYMS_API SYMS_SpatialMap1D* syms_group_unit_map(SYMS_Group *group);
+
+// thread safe (with lanes equipped to group)
+SYMS_API SYMS_SpatialMap1D* syms_group_proc_map_from_uid(SYMS_Group *group, SYMS_UnitID uid);
+SYMS_API SYMS_SpatialMap1D* syms_group_line_sequence_map_from_uid(SYMS_Group *group, SYMS_UnitID uid);
+
+SYMS_API SYMS_SpatialMap1D* syms_group_var_map_from_uid(SYMS_Group *group, SYMS_UnitID uid);
+
+SYMS_API void                syms_group_fetch_line_to_addr_maps_from_uid(SYMS_Group *group, SYMS_UnitID uid);
+SYMS_API SYMS_LineToAddrMap* syms_group_line_to_addr_map_from_uid_file_id(SYMS_Group *group, SYMS_UnitID uid,
+                                                                          SYMS_FileID file_id);
+
+//- accelerated versions
+SYMS_API SYMS_U64      syms_group_sec_number_from_voff__accelerated(SYMS_Group *group, SYMS_U64 voff);
+SYMS_API SYMS_U64      syms_group_sec_number_from_foff__accelerated(SYMS_Group *group, SYMS_U64 foff);
+
+SYMS_API SYMS_U64Maybe syms_group_sec_voff_from_foff__accelerated(SYMS_Group *group, SYMS_U64 voff);
+SYMS_API SYMS_U64Maybe syms_group_sec_foff_from_voff__accelerated(SYMS_Group *group, SYMS_U64 foff);
+
+SYMS_API SYMS_UnitID   syms_group_uid_from_voff__accelerated(SYMS_Group *group, SYMS_U64 voff);
+SYMS_API SYMS_SymbolID syms_group_proc_sid_from_uid_voff__accelerated(SYMS_Group *group, SYMS_UnitID uid,
+                                                                      SYMS_U64 voff);
+SYMS_API SYMS_SymbolID syms_group_var_sid_from_uid_voff__accelerated(SYMS_Group *group, SYMS_UnitID uid,
+                                                                     SYMS_U64 voff);
+
+SYMS_API SYMS_Line     syms_group_line_from_uid_voff__accelerated(SYMS_Group *group, SYMS_UnitID uid,
+                                                                  SYMS_U64 voff);
+SYMS_API SYMS_U64RangeArray syms_group_vranges_from_uid_line__accelerated(SYMS_Group *group, SYMS_UnitID uid,
+                                                                          SYMS_FileID file_id, SYMS_U32 line);
+
+//- line info binary search helper
+SYMS_API SYMS_U64      syms_index_from_n__u32__binary_search_round_up(SYMS_U32 *v, SYMS_U64 count, SYMS_U32 n);
+
+//- line-to-addr map helpers
+SYMS_API void          syms_line_to_addr_line_sort(SYMS_FileToLineToAddrLooseLine **array, SYMS_U64 count);
+SYMS_API void          syms_line_to_addr_line_sort__rec(SYMS_FileToLineToAddrLooseLine **array, SYMS_U64 count);
+
+////////////////////////////////
+// NOTE(allen): Syms Group Type Graph Functions
+
+SYMS_API SYMS_TypeNode* syms_group_type_from_usid(SYMS_Group *group, SYMS_USID usid);
+SYMS_API SYMS_TypeNode* syms_group_type_from_usid__rec(SYMS_Group *group, SYMS_USID usid);
+
+SYMS_API SYMS_TypeMemberArray* syms_group_members_from_type(SYMS_Group *group, SYMS_TypeNode *type);
+SYMS_API SYMS_EnumInfoArray*   syms_group_enum_members_from_type(SYMS_Group *group, SYMS_TypeNode *type);
+
+SYMS_API SYMS_U64       syms_group_type_size_from_usid(SYMS_Group *group, SYMS_USID usid);
+
+SYMS_API SYMS_TypeNode* syms_group_type_basic(SYMS_Group *group, SYMS_TypeKind basic_kind,
+                                              SYMS_U64 size, SYMS_String8 name);
+SYMS_API SYMS_TypeNode* syms_group_type_mod_from_type(SYMS_Group *group, SYMS_TypeNode *type,
+                                                      SYMS_TypeModifiers mods);
+SYMS_API SYMS_TypeNode* syms_group_type_ptr_from_type(SYMS_Group *group, SYMS_TypeKind ptr_kind,
+                                                      SYMS_TypeNode *type);
+SYMS_API SYMS_TypeNode* syms_group_type_array_from_type(SYMS_Group *group, SYMS_TypeNode *type, SYMS_U64 count);
+SYMS_API SYMS_TypeNode* syms_group_type_proc_from_type(SYMS_Group *group,
+                                                       SYMS_TypeNode *ret_type, SYMS_TypeNode *this_type,
+                                                       SYMS_TypeNode **param_types, SYMS_U64 param_count);
+SYMS_API SYMS_TypeNode* syms_group_type_member_ptr_from_type(SYMS_Group *group, SYMS_TypeNode *container,
+                                                             SYMS_TypeNode *type);
+
+SYMS_API SYMS_String8   syms_group_string_from_type(SYMS_Arena *arena, SYMS_TypeNode *type);
+
+SYMS_API void           syms_group_lhs_string_from_type(SYMS_Arena *arena, SYMS_TypeNode *type,
+                                                        SYMS_String8List *out);
+SYMS_API void           syms_group_rhs_string_from_type(SYMS_Arena *arena, SYMS_TypeNode *type,
+                                                        SYMS_String8List *out);
+
+SYMS_API void           syms_group_lhs_string_from_type_skip_return(SYMS_Arena *arena, SYMS_TypeNode *type,
+                                                                    SYMS_String8List *out);
+
+SYMS_API SYMS_U64       syms_type_usid_hash(SYMS_USID usid);
+SYMS_API SYMS_TypeNode* syms_type_from_usid(SYMS_TypeUSIDBuckets *buckets, SYMS_USID usid);
+SYMS_API void           syms_type_usid_buckets_insert(SYMS_Arena *arena, SYMS_TypeUSIDBuckets *buckets,
+                                                      SYMS_USID key, SYMS_TypeNode *type);
+
+SYMS_API SYMS_U64       syms_type_content_hash(SYMS_String8 data);
+SYMS_API SYMS_TypeNode* syms_type_from_content(SYMS_TypeContentBuckets *buckets, SYMS_String8 data);
+SYMS_API SYMS_String8   syms_type_content_buckets_insert(SYMS_Arena *arena, SYMS_TypeContentBuckets *buckets,
+                                                         SYMS_String8 key, SYMS_TypeNode *type);
+
+////////////////////////////////
+//~ NOTE(allen): Syms File Map
+
+SYMS_API SYMS_Name2FileIDMap* syms_group_file_map(SYMS_Group *group);
+
+
+#endif //SYMS_GROUP_H
diff --git a/Engine/Source/Runtime/SymsLib/syms/core/mach/syms_mach.c b/Engine/Source/Runtime/SymsLib/syms/core/mach/syms_mach.c
new file mode 100644
--- /dev/null
+++ b/Engine/Source/Runtime/SymsLib/syms/core/mach/syms_mach.c
@@ -0,0 +1,138 @@
+// Copyright Epic Games, Inc. All Rights Reserved.
+#ifndef SYMS_MACH_C
+#define SYMS_MACH_C
+
+////////////////////////////////
+//~ NOTE(allen): Generated
+
+#include "syms/core/generated/syms_meta_mach.c"
+
+////////////////////////////////
+//~ NOTE(allen): MACH Format Functions
+
+SYMS_API void
+syms_mach_fat_header_endian_swap_in_place(SYMS_MachFatHeader *x){
+  x->magic     = syms_bswap_u32(x->magic);
+  x->nfat_arch = syms_bswap_u32(x->nfat_arch);
+}
+
+SYMS_API void
+syms_mach_fat_arch_endian_swap_in_place(SYMS_MachFatArch *x){
+  x->cputype    = syms_bswap_u32(x->cputype);
+  x->cpusubtype = syms_bswap_u32(x->cpusubtype);
+  x->offset     = syms_bswap_u32(x->offset);
+  x->size       = syms_bswap_u32(x->size);
+  x->align      = syms_bswap_u32(x->align);
+}
+
+SYMS_API void
+syms_mach_header32_endian_swap_in_place(SYMS_MachHeader32 *x){
+  x->magic      = syms_bswap_u32(x->magic);
+  x->cputype    = syms_bswap_u32(x->cputype);
+  x->cpusubtype = syms_bswap_u32(x->cpusubtype);
+  x->filetype   = syms_bswap_u32(x->filetype);
+  x->ncmds      = syms_bswap_u32(x->ncmds);
+  x->sizeofcmds = syms_bswap_u32(x->sizeofcmds);
+  x->flags      = syms_bswap_u32(x->flags);
+}
+
+SYMS_API void
+syms_mach_header64_endian_swap_in_place(SYMS_MachHeader64 *x){
+  x->magic      = syms_bswap_u32(x->magic);
+  x->cputype    = syms_bswap_u32(x->cputype);
+  x->cpusubtype = syms_bswap_u32(x->cpusubtype);
+  x->filetype   = syms_bswap_u32(x->filetype);
+  x->ncmds      = syms_bswap_u32(x->ncmds);
+  x->sizeofcmds = syms_bswap_u32(x->sizeofcmds);
+  x->flags      = syms_bswap_u32(x->flags);
+  x->reserved   = syms_bswap_u32(x->reserved);
+}
+
+SYMS_API void
+syms_mach_segment_command32_endian_swap_in_place(SYMS_MachSegmentCommand32 *x){
+  x->cmd.type = syms_bswap_u32(x->cmd.type);
+  x->cmd.size = syms_bswap_u32(x->cmd.size);
+  x->vmaddr   = syms_bswap_u32(x->vmaddr);
+  x->vmsize   = syms_bswap_u32(x->vmsize);
+  x->fileoff  = syms_bswap_u32(x->fileoff);
+  x->filesize = syms_bswap_u32(x->filesize);
+  x->maxprot  = syms_bswap_u32(x->maxprot);
+  x->initprot = syms_bswap_u32(x->initprot);
+  x->nsects   = syms_bswap_u32(x->nsects);
+  x->flags    = syms_bswap_u32(x->flags);
+}
+
+SYMS_API void
+syms_mach_segment_command64_endian_swap_in_place(SYMS_MachSegmentCommand64 *x){
+  x->cmd.type = syms_bswap_u32(x->cmd.type);
+  x->cmd.size = syms_bswap_u32(x->cmd.size);
+  x->vmaddr   = syms_bswap_u64(x->vmaddr);
+  x->vmsize   = syms_bswap_u64(x->vmsize);
+  x->fileoff  = syms_bswap_u64(x->fileoff);
+  x->filesize = syms_bswap_u64(x->filesize);
+  x->maxprot  = syms_bswap_u32(x->maxprot);
+  x->initprot = syms_bswap_u32(x->initprot);
+  x->nsects   = syms_bswap_u32(x->nsects);
+  x->flags    = syms_bswap_u32(x->flags);
+}
+
+SYMS_API void
+syms_mach_section32_endian_swap_in_place(SYMS_MachSection32 *x){
+  x->addr         = syms_bswap_u32(x->addr);
+  x->size         = syms_bswap_u32(x->size);
+  x->offset       = syms_bswap_u32(x->offset);
+  x->align        = syms_bswap_u32(x->align);
+  x->relocoff     = syms_bswap_u32(x->relocoff);
+  x->nreloc       = syms_bswap_u32(x->nreloc);
+  x->flags        = syms_bswap_u32(x->flags);
+  x->reserved1    = syms_bswap_u32(x->reserved1);
+  x->reserved2    = syms_bswap_u32(x->reserved2);
+}
+
+SYMS_API void
+syms_mach_section64_endian_swap_in_place(SYMS_MachSection64 *x){
+  x->addr         = syms_bswap_u64(x->addr);
+  x->size         = syms_bswap_u64(x->size);
+  x->offset       = syms_bswap_u32(x->offset);
+  x->align        = syms_bswap_u32(x->align);
+  x->relocoff     = syms_bswap_u32(x->relocoff);
+  x->nreloc       = syms_bswap_u32(x->nreloc);
+  x->flags        = syms_bswap_u32(x->flags);
+  x->reserved1    = syms_bswap_u32(x->reserved1);
+  x->reserved2    = syms_bswap_u32(x->reserved2);
+}
+
+SYMS_API void
+syms_mach_header64_from_header32(SYMS_MachHeader64 *dst, SYMS_MachHeader32 *header32){
+  dst->magic = header32->magic;
+  dst->cputype = header32->cputype;
+  dst->cpusubtype = header32->cpusubtype;
+  dst->filetype = header32->filetype;
+  dst->ncmds = header32->ncmds;
+  dst->sizeofcmds = header32->sizeofcmds;
+  dst->flags = header32->flags;
+  dst->reserved = 0;
+}
+
+SYMS_API void
+syms_mach_nlist64_from_nlist32(SYMS_MachNList64 *dst, SYMS_MachNList32 *nlist32){
+  dst->n_strx = nlist32->n_strx;
+  dst->n_type = nlist32->n_type;
+  dst->n_sect = nlist32->n_sect;
+  dst->n_desc = nlist32->n_desc;
+  dst->n_value = nlist32->n_value;
+}
+
+SYMS_API SYMS_Arch
+syms_mach_arch_from_cputype(SYMS_MachCpuType cputype){
+  SYMS_Arch result = SYMS_Arch_Null;
+  switch (cputype){
+    case SYMS_MachCpuType_X86:    result = SYMS_Arch_X86;   break;
+    case SYMS_MachCpuType_X86_64: result = SYMS_Arch_X64;   break;
+    case SYMS_MachCpuType_ARM:    result = SYMS_Arch_ARM32; break;
+    case SYMS_MachCpuType_ARM64:  result = SYMS_Arch_ARM;   break;
+  }
+  return(result);
+}
+
+#endif // SYMS_MACH_C
diff --git a/Engine/Source/Runtime/SymsLib/syms/core/mach/syms_mach.h b/Engine/Source/Runtime/SymsLib/syms/core/mach/syms_mach.h
new file mode 100644
--- /dev/null
+++ b/Engine/Source/Runtime/SymsLib/syms/core/mach/syms_mach.h
@@ -0,0 +1,276 @@
+// Copyright Epic Games, Inc. All Rights Reserved.
+#ifndef SYMS_MACH_H
+#define SYMS_MACH_H
+
+/*
+** Reference materials:
+  ** machine.h https://opensource.apple.com/source/xnu/xnu-4570.41.2/osfmk/mach/machine.h.auto.html
+  ** loader.h  https://opensource.apple.com/source/xnu/xnu-4570.71.2/EXTERNAL_HEADERS/mach-o/loader.h.auto.html
+  ** nlist.h   https://opensource.apple.com/source/xnu/xnu-201/EXTERNAL_HEADERS/mach-o/nlist.h.auto.html
+  ** stab.h    https://opensource.apple.com/source/xnu/xnu-4903.241.1/EXTERNAL_HEADERS/mach-o/stab.h.auto.html
+*/
+
+////////////////////////////////
+//~ NOTE(allen): Generated
+
+#include "syms/core/generated/syms_meta_mach.h"
+
+////////////////////////////////
+//~ NOTE(allen): Constants from machine.h
+
+#define SYMS_MACH_CPU_ABI64     0x01000000
+#define SYMS_MACH_CPU_ABI64_MAS 0xff000000
+
+#define SYMS_MACH_CPU_SUBTYPE_MASK  0xff000000
+#define SYMS_MACH_CPU_SUBTYPE_LIB64 0x80000000
+
+////////////////////////////////
+//~ NOTE(allen): Types & Constants from loader.h
+
+//- magic numbers
+#define SYMS_MACH_MAGIC_32 0xFEEDFACE
+#define SYMS_MACH_CIGAM_32 0xCEFAEDFE
+
+#define SYMS_MACH_MAGIC_64 0xFEEDFACF
+#define SYMS_MACH_CIGAM_64 0xCFFAEDFE
+
+#define SYMS_MACH_FAT_MAGIC 0xCAFEBABE
+#define SYMS_MACH_FAT_CIGAM 0xBEBAFECA
+
+#define SYMS_MACH_IS_HEADER_32(x)      ((x) == SYMS_MACH_MAGIC_32 || (x) == SYMS_MACH_CIGAM_32)
+#define SYMS_MACH_IS_HEADER_64(x)      ((x) == SYMS_MACH_MAGIC_64 || (x) == SYMS_MACH_CIGAM_64)
+#define SYMS_MACH_HEADER_IS_SWAPPED(x) ((x) == SYMS_MACH_CIGAM_64 || (x) == SYMS_MACH_CIGAM_32)
+#define SYMS_MACH_IS_FAT(x)            ((x) == SYMS_MACH_FAT_MAGIC || (x) == SYMS_MACH_FAT_CIGAM)
+#define SYMS_MACH_FAT_IS_SWAPPED(x)    ((x) == SYMS_MACH_FAT_CIGAM)
+
+//- types
+typedef struct SYMS_MachLCStr{
+  SYMS_U32 offset;
+} SYMS_MachLCStr;
+
+typedef struct SYMS_MachFatHeader{
+  SYMS_U32 magic;
+  SYMS_U32 nfat_arch;
+} SYMS_MachFatHeader;
+
+typedef struct SYMS_MachFatArch{
+  SYMS_MachCpuType cputype;
+  SYMS_MachCpuSubtype cpusubtype;
+  SYMS_U32 offset;
+  SYMS_U32 size;
+  SYMS_U32 align;
+} SYMS_MachFatArch;
+
+typedef struct SYMS_MachHeader32{
+  SYMS_U32 magic;
+  SYMS_MachCpuType cputype;
+  SYMS_MachCpuSubtype cpusubtype;
+  SYMS_MachFiletype filetype;
+  SYMS_U32 ncmds;
+  SYMS_U32 sizeofcmds;
+  SYMS_MachFlags flags;
+} SYMS_MachHeader32;
+typedef struct SYMS_MachHeader64{
+  SYMS_U32 magic;
+  SYMS_MachCpuType cputype;
+  SYMS_MachCpuSubtype cpusubtype;
+  SYMS_MachFiletype filetype;
+  SYMS_U32 ncmds;
+  SYMS_U32 sizeofcmds;
+  SYMS_MachFlags flags;
+  SYMS_U32 reserved;
+} SYMS_MachHeader64;
+
+typedef struct SYMS_MachLoadCommand{
+  SYMS_MachLoadCommandType type;
+  SYMS_U32 size;
+} SYMS_MachLoadCommand;
+
+typedef SYMS_S32 SYMS_MachVMProt;
+
+typedef struct SYMS_MachSegmentCommand32{
+  SYMS_MachLoadCommand cmd;
+  SYMS_U8 segname[16];
+  SYMS_U32 vmaddr;
+  SYMS_U32 vmsize;
+  SYMS_U32 fileoff;
+  SYMS_U32 filesize;
+  SYMS_MachVMProt maxprot;
+  SYMS_MachVMProt initprot;
+  SYMS_U32 nsects;
+  SYMS_U32 flags;
+} SYMS_MachSegmentCommand32;
+
+typedef struct SYMS_MachSegmentCommand64{
+  SYMS_MachLoadCommand cmd;
+  SYMS_U8 segname[16];
+  SYMS_U64 vmaddr;
+  SYMS_U64 vmsize;
+  SYMS_U64 fileoff;
+  SYMS_U64 filesize;
+  SYMS_MachVMProt maxprot;
+  SYMS_MachVMProt initprot;
+  SYMS_U32 nsects;
+  SYMS_U32 flags;
+} SYMS_MachSegmentCommand64;
+
+typedef struct SYMS_MachSection32{
+  SYMS_U8 sectname[16];
+  SYMS_U8 segname[16];
+  SYMS_U32 addr;
+  SYMS_U32 size;
+  SYMS_U32 offset;
+  SYMS_U32 align;
+  SYMS_U32 relocoff;
+  SYMS_U32 nreloc;
+  SYMS_U32 flags;
+  SYMS_U32 reserved1;
+  SYMS_U32 reserved2;
+} SYMS_MachSection32;
+
+typedef struct SYMS_MachSection64{
+  SYMS_U8 sectname[16];
+  SYMS_U8 segname[16];
+  SYMS_U64 addr;
+  SYMS_U64 size;
+  SYMS_U32 offset;
+  SYMS_U32 align;
+  SYMS_U32 relocoff;
+  SYMS_U32 nreloc;
+  SYMS_U32 flags;
+  SYMS_U32 reserved1;
+  SYMS_U32 reserved2;
+} SYMS_MachSection64;
+
+// 128-bit unique number generated by linker to identify objects
+typedef struct SYMS_MachUUID{
+  SYMS_U32 cmd;
+  SYMS_U32 cmdsize;
+  SYMS_U8 uuid[16];
+} SYMS_MachUUID;
+
+typedef struct SYMS_MachDylib{
+  SYMS_MachLCStr name;
+  SYMS_U32 timestamp;
+  SYMS_U32 current_version;
+  SYMS_U32 compatability_version;
+} SYMS_MachDylib;
+
+typedef struct SYMS_MachDylibCommand{
+  SYMS_U32 cmd;
+  SYMS_U32 cmdsize;
+  SYMS_MachDylib dylib;
+} SYMS_MachDylibCommand;
+
+typedef struct SYMS_MachDyldInfoCommand{
+  SYMS_U32 cmd;
+  SYMS_U32 cmdsize;
+  SYMS_U32 rebase_off;
+  SYMS_U32 rebase_size;
+  SYMS_U32 bind_off;
+  SYMS_U32 bind_size;
+  SYMS_U32 weak_bind_off;
+  SYMS_U32 weak_bind_size;
+  SYMS_U32 lazy_bind_off;
+  SYMS_U32 lazy_bind_size;
+  SYMS_U32 export_off;
+  SYMS_U32 export_size;
+} SYMS_MachDyldInfoCommand;
+
+// command for symbol table, format for entries is defined in list.h and stab.h
+typedef struct SYMS_MachSymtabCommand{
+  SYMS_U32 cmd;
+  SYMS_U32 cmdsize;
+  SYMS_U32 symoff;  // offset from the image base to array of SYMS_MachNList32
+  SYMS_U32 nsyms;   // symbol count
+  SYMS_U32 stroff;  // offset from the image base to string table
+  SYMS_U32 strsize; // string table size
+} SYMS_MachSymtabCommand;
+
+typedef struct SYMS_MachDySymtabCommand{
+  SYMS_U32 cmd;
+  SYMS_U32 cmdsize;
+  SYMS_U32 ilocalsym;
+  SYMS_U32 nlocalsym;
+  SYMS_U32 iextdefsym;
+  SYMS_U32 nextdefsym;
+  SYMS_U32 iundefsym;
+  SYMS_U32 nundefsym;
+  SYMS_U32 tocoff;
+  SYMS_U32 ntoc;
+  SYMS_U32 modtaboff;
+  SYMS_U32 nmodtab;
+  SYMS_U32 extrefsymoff;
+  SYMS_U32 nextrefsyms;
+  SYMS_U32 indirectsymoff;
+  SYMS_U32 nindirectsyms;
+  SYMS_U32 extreloff;
+  SYMS_U32 nextrel;
+  SYMS_U32 locreloff;
+  SYMS_U32 nlocrel;
+} SYMS_MachDySymtabCommand;
+
+#define SYMS_MACH_VERSION_NIBBLE_A(x) (((x) & 0xFFFF0000) >> 16)
+#define SYMS_MACH_VERSION_NIBBLE_B(x) (((x) & 0x0000FF00) >> 8)
+#define SYMS_MACH_VERSION_NIBBLE_C(x) (((x) & 0x000000FF) >> 0)
+
+typedef struct SYMS_MachBuildVersionCommand{
+  SYMS_U32 cmd;
+  SYMS_U32 cmdsize;
+  SYMS_MachPlatformType platform;
+  SYMS_U32 minos;
+  SYMS_U32 sdk;
+  SYMS_U32 ntools;
+} SYMS_MachBuildVersionCommand;
+
+typedef struct SYMS_MachBuildToolVersion{
+  SYMS_MachToolType tool;
+  SYMS_U32 version;
+} SYMS_MachBuildToolVersion;
+
+////////////////////////////////
+//~ NOTE(allen): Types & Constants from nlist.h
+
+#define SYMS_MACH_NLIST_STAB(x) ((x) & 0xE0)
+#define SYMS_MACH_NLIST_PEXT(x) ((x) & 0x10)
+#define SYMS_MACH_NLIST_TYPE(x) ((x) & 0x0E)
+#define SYMS_MACH_NLIST_EXT(x)  ((x) & 0x01)
+
+typedef struct SYMS_MachNList32{
+  SYMS_U32 n_strx;
+  SYMS_U8 n_type;
+  SYMS_U8 n_sect;
+  SYMS_U16 n_desc;
+  SYMS_U32 n_value;
+} SYMS_MachNList32;
+
+typedef struct SYMS_MachNList64{
+  SYMS_U32 n_strx;
+  SYMS_U8 n_type;
+  SYMS_U8 n_sect;
+  SYMS_U16 n_desc;
+  SYMS_U64 n_value;
+} SYMS_MachNList64;
+
+////////////////////////////////
+//~ NOTE(allen): MACH Format Functions
+
+SYMS_C_LINKAGE_BEGIN
+
+SYMS_API void      syms_mach_fat_header_endian_swap_in_place(SYMS_MachFatHeader *x);
+SYMS_API void      syms_mach_fat_arch_endian_swap_in_place(SYMS_MachFatArch *x);
+SYMS_API void      syms_mach_header32_endian_swap_in_place(SYMS_MachHeader32 *x);
+SYMS_API void      syms_mach_header64_endian_swap_in_place(SYMS_MachHeader64 *x);
+SYMS_API void      syms_mach_segment_command32_endian_swap_in_place(SYMS_MachSegmentCommand32 *x);
+SYMS_API void      syms_mach_segment_command64_endian_swap_in_place(SYMS_MachSegmentCommand64 *x);
+SYMS_API void      syms_mach_section32_endian_swap_in_place(SYMS_MachSection32 *section32);
+SYMS_API void      syms_mach_section64_endian_swap_in_place(SYMS_MachSection64 *section64);
+
+SYMS_API void      syms_mach_header64_from_header32(SYMS_MachHeader64 *dst, SYMS_MachHeader32 *header32);
+SYMS_API void      syms_mach_nlist64_from_nlist32(SYMS_MachNList64 *dst, SYMS_MachNList32 *nlist32);
+
+SYMS_API SYMS_Arch syms_mach_arch_from_cputype(SYMS_MachCpuType cputype);
+
+SYMS_C_LINKAGE_END
+
+#endif // SYMS_MACH_H
