Update syms library to latest version, includes dwarf5 support

diff --git a/Engine/Source/Runtime/SymsLib/syms/core/pdb/syms_pdb_parser.c b/Engine/Source/Runtime/SymsLib/syms/core/pdb/syms_pdb_parser.c
new file mode 100644
--- /dev/null
+++ b/Engine/Source/Runtime/SymsLib/syms/core/pdb/syms_pdb_parser.c
@@ -0,0 +1,4860 @@
+// Copyright Epic Games, Inc. All Rights Reserved.
+#ifndef SYMS_PDB_PARSER_C
+#define SYMS_PDB_PARSER_C
+
+////////////////////////////////
+//~ NOTE(allen): PDB's Prerequisite Formats
+
+#include "syms/core/pdb/syms_msf_parser.c"
+
+////////////////////////////////
+//~ NOTE(allen): PDB TPI Functions
+
+SYMS_API SYMS_PdbTpiOffRange
+syms_pdb_tpi__hint_from_index(SYMS_PdbTpiAccel *tpi, SYMS_CvTypeIndex ti){
+  // b.s. rule: largest 'index' so that 'tpi->hints[index].ti <= ti'
+  //            or, when that's impossible, 'index = tpi->count'
+  SYMS_PdbTpiOffHint *hints = tpi->hints;
+  SYMS_U32 index = tpi->count;
+  SYMS_U32 min = 0;
+  SYMS_U32 max = tpi->count;
+  for (;;){
+    SYMS_U32 mid = (min + max)/2;
+    SYMS_PdbTpiOffHint *h = hints + mid;
+    if (h->ti > ti){
+      max = mid;
+      if (max == 0){
+        break;
+      }
+    }
+    else{
+      min = mid;
+      if (min + 1 >= max){
+        index = min;
+        break;
+      }
+    }
+  }
+  
+  // return result
+  SYMS_PdbTpiOffRange result = {0};
+  if (index < tpi->count){
+    result.first_off = tpi->base_off + tpi->hints[index].off;
+    result.first_ti = tpi->hints[index].ti;
+    result.opl_ti = SYMS_U32_MAX;
+    if (index + 1 < tpi->count){
+      result.opl_ti = tpi->hints[index + 1].ti;
+    }
+  }
+  return(result);
+}
+
+SYMS_API void
+syms_pdb_tpi__fill_off_range(SYMS_String8 data, SYMS_MsfAccel *msf, SYMS_PdbTpiAccel *tpi,
+                             SYMS_PdbTpiOffRange *fill){
+  // fill params
+  SYMS_U32         first_off = fill->first_off;
+  SYMS_CvTypeIndex first_ti  = SYMS_ClampBot(tpi->first_ti, fill->first_ti);
+  SYMS_CvTypeIndex opl_ti    = SYMS_ClampTop(fill->opl_ti, tpi->opl_ti);
+  SYMS_MsfStreamNumber sn    = tpi->type_sn;
+  
+  // default fill to non-zero invalid off (all bytes 0xFF)
+  //  we do this because we use zero to indicate "cache not filled"
+  //  so "filled by invalid" needs another value.
+  SYMS_U32 *type_off_ptr = tpi->off + first_ti - tpi->first_ti;
+  syms_memset(type_off_ptr, 0xFF, sizeof(*type_off_ptr)*(opl_ti - first_ti));
+  
+  // fill
+  SYMS_U32 cursor = first_off;
+  for (SYMS_CvTypeIndex ti = first_ti;
+       ti < opl_ti;
+       ti += 1, type_off_ptr += 1){
+    // save off
+    *type_off_ptr = cursor;
+    // read size
+    SYMS_U16 size = 0;
+    syms_msf_read_struct(data, msf, sn, cursor, &size);
+    // advance
+    cursor += 2 + size;
+    if (size == 0 || !syms_msf_bounds_check(msf, sn, cursor)){
+      break;
+    }
+  }
+}
+
+SYMS_API SYMS_PdbTpiAccel
+syms_pdb_tpi_accel_from_sn(SYMS_Arena *arena, SYMS_String8 data, SYMS_MsfAccel *msf, SYMS_MsfStreamNumber sn){
+  SYMS_ProfBegin("syms_pdb_tpi_accel_from_sn");
+  
+  // init data structures
+  SYMS_PdbTpiAccel result = {0};
+  if (sn == SYMS_PdbFixedStream_TPI ||
+      sn == SYMS_PdbFixedStream_IPI){
+    SYMS_PdbTpiHeader hdr = {0};
+    syms_msf_read_struct(data, msf, sn, 0, &hdr);
+    if (hdr.version == SYMS_PdbTpiVersion_IMPV80){ // NOTE(allen): the only version we support so far.
+      
+      // parse hash table
+      SYMS_PdbChain **buckets = 0;
+      SYMS_U32 bucket_count = 0;
+      
+      if (hdr.hash_bucket_count != 0){
+        // grab scratch
+        SYMS_ArenaTemp scratch = syms_get_scratch(&arena, 1);
+        
+        // setup buckets
+        bucket_count = hdr.hash_bucket_count;
+        buckets = syms_push_array(arena, SYMS_PdbChain*, bucket_count);
+        syms_memset(buckets, 0, sizeof(*buckets)*bucket_count);
+        
+        // grab hash array
+        SYMS_MsfRange hash_range = syms_msf_range_from_sn(msf, hdr.hash_sn);
+        SYMS_MsfRange range = syms_msf_sub_range(hash_range, hdr.hash_vals.off, hdr.hash_vals.size);
+        SYMS_String8 hash_array = syms_msf_read_whole_range(scratch.arena, data, msf, range);
+        
+        // iterate hash vals
+        SYMS_U32 stride = hdr.hash_key_size;
+        SYMS_U8 *cursor = hash_array.str;
+        SYMS_U8 *opl_cursor = hash_array.str + hash_array.size;
+        SYMS_CvTypeIndex ti = hdr.ti_lo;
+        for (;;){
+          // read index
+          SYMS_U32 bucket_index = 0;
+          syms_memmove(&bucket_index, cursor, stride);
+          
+          // save to typemap
+          SYMS_PdbChain *entry = 0;
+          if (bucket_index < bucket_count){
+            entry = syms_push_array(arena, SYMS_PdbChain, 1);
+            syms_memzero_struct(entry);
+            SYMS_StackPush(buckets[bucket_index], entry);
+            entry->v = ti;
+          }
+          
+          // advance cursor
+          cursor += stride;
+          ti += 1;
+          
+          // exit condition
+          if (cursor >= opl_cursor){
+            break;
+          }
+        }
+        
+        // release scratch
+        syms_release_scratch(scratch);
+      }
+      
+      // parse hint table
+      SYMS_PdbTpiOffHint *hints = 0;
+      SYMS_U32 hint_count = 0;
+      {
+        SYMS_MsfRange hash_range = syms_msf_range_from_sn(msf, hdr.hash_sn);
+        SYMS_MsfRange range = syms_msf_sub_range(hash_range, hdr.ti_off.off, hdr.ti_off.size);
+        SYMS_String8 hint_memory = syms_msf_read_whole_range(arena, data, msf, range);
+        hints = (SYMS_PdbTpiOffHint*)hint_memory.str;
+        hint_count = (hint_memory.size)/sizeof(*hints);
+      }
+      
+      // offset array
+      SYMS_CvTypeIndex first_ti = 0;
+      SYMS_CvTypeIndex opl_ti = 0;
+      SYMS_U32 *off = 0;
+      {
+        first_ti = hdr.ti_lo;
+        opl_ti   = SYMS_ClampBot(first_ti, hdr.ti_hi);
+        
+        SYMS_U64 count = opl_ti - first_ti;
+        off = syms_push_array(arena, SYMS_U32, count);
+        syms_memset(off, 0, sizeof(*off)*count);
+      }
+      
+      // fill result
+      result.type_sn  = sn;
+      result.buckets  = buckets;
+      result.bucket_count = bucket_count;
+      result.count    = hint_count;
+      result.hints    = hints;
+      result.base_off = hdr.header_size;
+      result.first_ti = first_ti;
+      result.opl_ti   = opl_ti;
+      result.off      = off;
+    }
+  }
+  
+  SYMS_ProfEnd();
+  
+  return(result);
+}
+
+SYMS_API SYMS_U32
+syms_pdb_tpi_off_from_ti(SYMS_String8 data, SYMS_MsfAccel *msf, SYMS_PdbTpiAccel *tpi, SYMS_CvTypeIndex ti){
+  SYMS_U32 result = 0;
+  if (tpi->first_ti <= ti && ti < tpi->opl_ti){
+    // read off
+    SYMS_U32 relative_index = ti - tpi->first_ti;
+    SYMS_U32 off = tpi->off[relative_index];
+    
+    // if not zero this is the result
+    if (off != 0){
+      result = off;
+    }
+    
+    // otherwise compute and cache this off range
+    else{
+      SYMS_PdbTpiOffRange hint = syms_pdb_tpi__hint_from_index(tpi, ti);
+      syms_pdb_tpi__fill_off_range(data, msf, tpi, &hint);
+      result = tpi->off[relative_index];
+    }
+    
+    // collapse "error cases" down to null
+    if (result == SYMS_U32_MAX){
+      result = 0;
+    }
+  }
+  return(result);
+}
+
+SYMS_API SYMS_U32
+syms_pdb_tpi_base_off(SYMS_PdbTpiAccel *tpi){
+  return(tpi->base_off);
+}
+
+SYMS_API SYMS_MsfRange
+syms_pdb_tpi_range(SYMS_MsfAccel *msf, SYMS_PdbTpiAccel *map){
+  SYMS_MsfRange range = syms_msf_range_from_sn(msf, map->type_sn);
+  SYMS_MsfRange result = syms_msf_sub_range(range, map->base_off, range.size - map->base_off);
+  return(result);
+}
+
+SYMS_API SYMS_USIDList
+syms_pdb_types_from_name(SYMS_Arena *arena, SYMS_String8 data, SYMS_PdbDbgAccel *dbg, SYMS_PdbUnitAccel *unit,
+                         SYMS_String8 name){
+  // setup accel
+  SYMS_PdbTpiAccel *tpi = &dbg->tpi;
+  
+  SYMS_USIDList result = {0};
+  if (tpi->bucket_count > 0){
+    
+    // get bucket
+    SYMS_U32 name_hash = syms_pdb_hashV1(name);
+    SYMS_U32 bucket_index = name_hash%tpi->bucket_count;
+    
+    // iterate bucket
+    for (SYMS_PdbChain *bucket = tpi->buckets[bucket_index];
+         bucket != 0;
+         bucket = bucket->next){
+      SYMS_CvTypeIndex ti = bucket->v;
+      if (unit->top_min_index <= ti){
+        SYMS_U64 index = ti - unit->top_min_index;
+        if (index < unit->ti_count){
+          SYMS_PdbStub *stub = unit->ti_stubs + index;
+          if (stub != 0 && syms_string_match(stub->name, name, 0)){
+            SYMS_USIDNode *usid_node = syms_push_array(arena, SYMS_USIDNode, 1);
+            SYMS_QueuePush(result.first, result.last, usid_node);
+            result.count += 1;
+            usid_node->usid.uid = SYMS_PdbPseudoUnit_TPI;
+            usid_node->usid.sid = SYMS_ID_u32_u32(SYMS_PdbSymbolIDKind_Index, ti);
+          }
+        }
+      }
+    }
+  }
+  
+  return(result);
+}
+
+////////////////////////////////
+//~ NOTE(allen): PDB GSI
+
+SYMS_API SYMS_PdbGsiAccel
+syms_pdb_gsi_accel_from_range(SYMS_Arena *arena, SYMS_String8 data, SYMS_MsfAccel *accel, SYMS_MsfRange range){
+  SYMS_ProfBegin("syms_pdb_parse_info");
+  
+  // get scratch
+  SYMS_ArenaTemp scratch = syms_get_scratch(&arena, 1);
+  // check gsi info
+  SYMS_B32 has_gsi = syms_false;
+  SYMS_PdbGsiHeader header = {0};
+  if (syms_msf_read_struct_in_range(data, accel, range, 0, &header)){
+    if (header.sig == SYMS_PdbGsiSignature_Basic &&
+        header.ver == SYMS_PdbGsiVersion_V70 &&
+        header.num_buckets != 0){
+      has_gsi = syms_true;
+    }
+  }
+  
+  // hash offset
+  SYMS_U32 hash_record_array_off = sizeof(header);
+  
+  // bucket count
+  SYMS_U32 bucket_count = 4096;
+  SYMS_U32 slot_count = bucket_count + 1;
+  
+  // get unpacked offsets
+  SYMS_U32 *unpacked_offsets = 0;
+  if (has_gsi){
+    unpacked_offsets = syms_push_array(scratch.arena, SYMS_U32, slot_count);
+    
+    SYMS_ArenaTemp temp = syms_arena_temp_begin(scratch.arena);
+    
+    // array offsets
+    SYMS_U32 bitmask_count = ((slot_count + 31)/32);
+    SYMS_U32 bitmask_size = bitmask_count*4;
+    SYMS_U32 bitmask_off = hash_record_array_off + header.hr_len;
+    SYMS_U32 offsets_off = bitmask_off + bitmask_size;
+    
+    // read bitmask
+    SYMS_MsfRange bitmask_range = syms_msf_sub_range(range, bitmask_off, bitmask_size);
+    SYMS_String8 bitmask = syms_msf_read_whole_range(scratch.arena, data, accel, bitmask_range);
+    SYMS_U32 bitmask_count_clamped = bitmask.size/4;
+    
+    // read packed offsets
+    SYMS_MsfRange offsets_range = syms_msf_sub_range(range, offsets_off, range.size - offsets_off);
+    SYMS_String8 packed_offsets = syms_msf_read_whole_range(scratch.arena, data, accel, offsets_range);
+    SYMS_U32 packed_offsets_count = packed_offsets.size/4;
+    
+    // unpack offsets
+    SYMS_U32 *bitmask_ptr = (SYMS_U32*)(bitmask.str);
+    SYMS_U32 *bitmask_opl = bitmask_ptr + bitmask_count_clamped;
+    SYMS_U32 *dst_ptr = unpacked_offsets;
+    SYMS_U32 *dst_opl = dst_ptr + slot_count;
+    SYMS_U32 *src_ptr = (SYMS_U32*)packed_offsets.str;
+    SYMS_U32 *src_opl = src_ptr + packed_offsets_count;
+    for (; bitmask_ptr < bitmask_opl && src_ptr < src_opl; bitmask_ptr += 1){
+      SYMS_U32 bits = *bitmask_ptr;
+      SYMS_U32 dst_max = (SYMS_U32)(dst_opl - dst_ptr);
+      SYMS_U32 src_max = (SYMS_U32)(src_opl - src_ptr);
+      SYMS_U32 k_max0 = SYMS_ClampTop(32, dst_max);
+      SYMS_U32 k_max = SYMS_ClampTop(k_max0, src_max);
+      for (SYMS_U32 k = 0; k < k_max; k += 1){
+        if ((bits & 1) == 1){
+          *dst_ptr = *src_ptr;
+          src_ptr += 1;
+        }
+        else{
+          *dst_ptr = SYMS_U32_MAX;
+        }
+        dst_ptr += 1;
+        bits >>= 1;
+      }
+    }
+    for (; dst_ptr < dst_opl; dst_ptr += 1){
+      *dst_ptr = SYMS_U32_MAX;
+    }
+    
+    syms_arena_temp_end(temp);
+  }
+  
+  // construct table
+  SYMS_PdbChain **table = 0;
+  SYMS_B32 table_build_success = syms_true;
+  if (unpacked_offsets != 0){
+    table = syms_push_array(arena, SYMS_PdbChain*, bucket_count + 1);
+    
+    // read hash records
+    SYMS_MsfRange hash_record_range = syms_msf_sub_range(range, hash_record_array_off, header.hr_len);
+    SYMS_String8 hash_records = syms_msf_read_whole_range(scratch.arena, data, accel, hash_record_range);
+    SYMS_U32 entry_count = hash_records.size/sizeof(SYMS_PdbGsiHashRecord);
+    
+    // build bucket chains
+    // NOTE(allen):  we have to do this backwards to make the singly linked
+    // stacks end up in the right order, and because of the way groups of
+    // offsets are stored by PDB (it is specified this way on purpose)
+    SYMS_U32 last_index = entry_count - 1;
+    SYMS_PdbGsiHashRecord *in_ptr  = (SYMS_PdbGsiHashRecord*)(hash_records.str) + last_index;
+    SYMS_U32 prev_n = entry_count;
+    for (SYMS_U32 i = bucket_count + 1; i > 0;){
+      i -= 1;
+      SYMS_PdbChain *table_ptr = 0;
+      if (unpacked_offsets[i] != SYMS_U32_MAX){
+        SYMS_U32 n = unpacked_offsets[i]/sizeof(SYMS_PdbGsiHrOffsetCalc);
+        if (n > prev_n){
+          table_build_success = syms_false;
+          break;
+        }
+        SYMS_U32 num_steps = prev_n - n;
+        for (SYMS_U32 j = 0; j < num_steps; j += 1){
+          SYMS_PdbChain *bucket = syms_push_array(arena, SYMS_PdbChain, 1);
+          SYMS_StackPush(table_ptr, bucket);
+          bucket->v = in_ptr->off - 1;
+          in_ptr -= 1;
+        }
+        prev_n = n;
+      }
+      table[i] = table_ptr;
+    }
+  }
+  
+  // build the result
+  SYMS_PdbGsiAccel result = {0};
+  if (table_build_success){
+    result.buckets = table;
+    result.bucket_count = bucket_count;
+  }
+  
+  // release scratch
+  syms_release_scratch(scratch);
+  
+  SYMS_ProfEnd();
+  
+  return(result);
+}
+
+SYMS_API SYMS_MsfRange
+syms_pdb_gsi_part_from_psi_range(SYMS_MsfRange psi_range){
+  // NOTE(allen): We don't actually use the PSI specific information right now,
+  // so it's basically just a wrapper around GSI. We didn't collapse it down to
+  // one system because there might be reasons to start using these extra features
+  // and then this will need to be a wrapper.
+#if 0
+  // psi header
+  SYMS_PdbPsiHeader header = {0};
+  syms_msf_read_struct_in_range(data, accel, psi_range, 0, &header);
+#endif
+  
+  SYMS_MsfRange result = syms_msf_sub_range(psi_range, sizeof(SYMS_PdbPsiHeader),
+                                            psi_range.size - sizeof(SYMS_PdbPsiHeader));
+  return(result);
+}
+
+SYMS_API SYMS_USIDList
+syms_pdb_symbols_from_name(SYMS_Arena *arena, SYMS_String8 data, SYMS_MsfAccel *msf,
+                           SYMS_PdbGsiAccel *gsi, SYMS_PdbUnitAccel *unit, SYMS_String8 name){
+  SYMS_USIDList result = {0};
+  
+  if (gsi->bucket_count > 0){
+    // get unit range
+    SYMS_MsfRange range = syms_msf_range_from_sn(msf, unit->sn);
+    
+    // grab scratch
+    SYMS_ArenaTemp scratch = syms_get_scratch(&arena, 1);
+    
+    // get bucket
+    SYMS_U32 name_hash = syms_pdb_hashV1(name);
+    SYMS_U32 bucket_index = name_hash%gsi->bucket_count;
+    
+    // iterate bucket
+    for (SYMS_PdbChain *bucket = gsi->buckets[bucket_index];
+         bucket != 0;
+         bucket = bucket->next){
+      SYMS_U32 off = bucket->v;
+      SYMS_PdbStub *stub = syms_pdb_stub_from_unit_off(unit, off);
+      if (stub != 0 && syms_string_match(stub->name, name, 0)){
+        SYMS_CvElement element = syms_cv_element(data, msf, range, stub->off);
+        
+        SYMS_UnitID   uid = 0;
+        SYMS_SymbolID sid = 0;
+        switch (element.kind){
+          case SYMS_CvSymKind_PROCREF:
+          case SYMS_CvSymKind_LPROCREF:
+          case SYMS_CvSymKind_DATAREF:
+          {
+            SYMS_CvRef2 ref2 = {0};
+            syms_msf_read_struct_in_range(data, msf, element.range, 0, &ref2);
+            uid = ref2.imod + SYMS_PdbPseudoUnit_FIRST_COMP_UNIT;
+            sid = SYMS_ID_u32_u32(SYMS_PdbSymbolIDKind_Off, ref2.sym_off);
+          }break;
+        }
+        
+        if (uid != 0){
+          SYMS_USIDNode *chain = syms_push_array(arena, SYMS_USIDNode, 1);
+          SYMS_QueuePush(result.first, result.last, chain);
+          result.count += 1;
+          chain->usid.uid = uid;
+          chain->usid.sid = sid;
+        }
+      }
+    }
+    
+    // release scratch
+    syms_release_scratch(scratch);
+  }
+  
+  return(result);
+}
+
+////////////////////////////////
+//~ NOTE(allen): PDB Accel Functions
+
+// pdb specific parsing
+
+SYMS_API SYMS_PdbDbiAccel
+syms_pdb_dbi_from_msf(SYMS_String8 data, SYMS_MsfAccel *accel){
+  SYMS_ProfBegin("syms_pdb_dbi_from_msf");
+  
+  // grab the dbi stream
+  SYMS_MsfRange range = syms_msf_range_from_sn(accel, SYMS_PdbFixedStream_DBI);
+  
+  // read dbi header
+  SYMS_PdbDbiHeader header = {0};
+  syms_msf_read_struct_in_range(data, accel, range, 0, &header);
+  
+  // extract dbi info
+  SYMS_PdbDbiAccel result = {0};
+  if (header.sig == SYMS_PdbDbiHeaderSignature_V1){
+    // info directly from the header
+    result.valid = 1;
+    result.machine_type = header.machine;
+    result.gsi_sn = header.gsi_sn;
+    result.psi_sn = header.psi_sn;
+    result.sym_sn = header.sym_sn;
+    
+    // organize the sizes of the ranges in dbi
+    SYMS_U64 range_size[SYMS_PdbDbiRange_COUNT];
+    range_size[SYMS_PdbDbiRange_ModuleInfo] = header.module_info_size;
+    range_size[SYMS_PdbDbiRange_SecCon]     = header.sec_con_size;
+    range_size[SYMS_PdbDbiRange_SecMap]     = header.sec_map_size;
+    range_size[SYMS_PdbDbiRange_FileInfo]   = header.file_info_size;
+    range_size[SYMS_PdbDbiRange_TSM]        = header.tsm_size;
+    range_size[SYMS_PdbDbiRange_EcInfo]     = header.ec_info_size;
+    range_size[SYMS_PdbDbiRange_DbgHeader]  = header.dbg_header_size;
+    
+    // fill range offset array
+    {
+      SYMS_U64 cursor = sizeof(header);
+      SYMS_U64 i = 0;
+      for (; i < (SYMS_U64)(SYMS_PdbDbiRange_COUNT); i += 1){
+        result.range_off[i] = cursor;
+        cursor += range_size[i];
+        cursor = SYMS_ClampTop(cursor, range.size);
+      }
+      //~ NOTE(allen): one last value past the end so that off[i + 1] - off[i] can get us the sizes back.
+      result.range_off[i] = cursor;
+    }
+    
+    // read debug streams
+    {
+      //~ NOTE(allen): zero is a valid stream [ it sucks I know :( ]
+      // so we explicitly invalidate sn by clearing to 0xFF instead.
+      syms_memset(result.dbg_sn, 0xFF, sizeof(result.dbg_sn));
+      SYMS_MsfRange dbg_sub_range = syms_pdb_dbi_sub_range(&result, SYMS_PdbDbiRange_DbgHeader);
+      SYMS_U64 read_size = SYMS_ClampTop(sizeof(result.dbg_sn), dbg_sub_range.size);
+      syms_msf_read_in_range(data, accel, dbg_sub_range, 0, read_size, result.dbg_sn);
+    }
+  }
+  
+  SYMS_ProfEnd();
+  
+  return(result);
+}
+
+SYMS_API SYMS_MsfRange
+syms_pdb_dbi_sub_range(SYMS_PdbDbiAccel *dbi, SYMS_PdbDbiRange n){
+  SYMS_MsfRange result = {0};
+  result.sn = SYMS_PdbFixedStream_DBI;
+  result.off = dbi->range_off[n];
+  result.size = dbi->range_off[n + 1] - result.off;
+  return(result);
+}
+
+SYMS_API SYMS_MsfRange
+syms_pdb_dbi_stream(SYMS_MsfAccel *msf, SYMS_PdbDbiAccel *dbi, SYMS_PdbDbiStream n){
+  SYMS_MsfRange result = syms_msf_range_from_sn(msf, dbi->dbg_sn[n]);
+  return(result);
+}
+
+SYMS_API SYMS_PdbInfoTable
+syms_pdb_parse_info(SYMS_Arena *arena, SYMS_String8 data, SYMS_MsfAccel *accel){
+  SYMS_ProfBegin("syms_pdb_parse_info");
+  
+  // grab the info stream
+  SYMS_MsfRange range = syms_msf_range_from_sn(accel, SYMS_PdbFixedStream_PDB);
+  
+  // read info stream's header
+  SYMS_PdbInfoHeader header = {0};
+  syms_msf_read_struct_in_range(data, accel, range, 0, &header);
+  SYMS_U32 after_header_off = sizeof(header);
+  
+  // read auth_guid (only in certain recent pdb versions)
+  SYMS_PeGuid auth_guid = {0};
+  SYMS_U32 after_auth_guid_off = after_header_off;
+  switch (header.version){
+    case SYMS_PdbVersion_VC70_DEP:
+    case SYMS_PdbVersion_VC70:
+    case SYMS_PdbVersion_VC80:
+    case SYMS_PdbVersion_VC110:
+    case SYMS_PdbVersion_VC140:
+    {
+      syms_msf_read_struct_in_range(data, accel, range, after_header_off, &auth_guid);
+      after_auth_guid_off = after_header_off + sizeof(auth_guid);
+    }break;
+    
+    default:
+    {}break;
+  }
+  
+  // read table layout data
+  SYMS_U32 names_base_off = 0;
+  SYMS_U32 hash_table_count = 0;
+  SYMS_U32 hash_table_max = 0;
+  SYMS_U32 epilogue_base_off = 0;
+  
+  if (header.version != 0){
+    SYMS_U32 names_len_off = after_auth_guid_off;
+    
+    SYMS_U32 names_len = 0;
+    syms_msf_read_struct_in_range(data, accel, range, names_len_off, &names_len);
+    
+    names_base_off = names_len_off + 4;
+    
+    SYMS_U32 hash_table_count_off = names_base_off + names_len;
+    syms_msf_read_struct_in_range(data, accel, range, hash_table_count_off, &hash_table_count);
+    
+    SYMS_U32 hash_table_max_off = hash_table_count_off + 4;
+    syms_msf_read_struct_in_range(data, accel, range, hash_table_max_off, &hash_table_max);
+    
+    SYMS_U32 num_present_words_off = hash_table_max_off + 4;
+    SYMS_U32 num_present_words = 0;
+    syms_msf_read_struct_in_range(data, accel, range, num_present_words_off, &num_present_words);
+    
+    SYMS_U32 present_words_array_off = num_present_words_off + 4;
+    
+    SYMS_U32 num_deleted_words_off = present_words_array_off + num_present_words*sizeof(SYMS_U32);
+    SYMS_U32 num_deleted_words = 0;
+    syms_msf_read_struct_in_range(data, accel, range, num_deleted_words_off, &num_deleted_words);
+    
+    SYMS_U32 deleted_words_array_off = num_deleted_words_off + 4;
+    
+    epilogue_base_off = deleted_words_array_off + num_deleted_words*sizeof(SYMS_U32);
+  }
+  
+  // read table
+  SYMS_PdbInfoTable result = {0};
+  if (hash_table_count > 0 && syms_msf_bounds_check_in_range(range, epilogue_base_off)){
+    SYMS_U32 record_off = epilogue_base_off;
+    for (SYMS_U32 i = 0;
+         i < hash_table_count;
+         i += 1, record_off += 8){
+      // read record
+      SYMS_U32 record[2];
+      syms_msf_read_in_range(data, accel, range, record_off, 8, record);
+      SYMS_U32 relative_name_off = record[0];
+      SYMS_MsfStreamNumber sn = (SYMS_MsfStreamNumber)record[1];
+      
+      // read name
+      SYMS_U32 name_off = names_base_off + relative_name_off;
+      SYMS_String8 name = syms_msf_read_zstring_in_range(arena, data, accel, range, name_off);
+      
+      // push new slot
+      SYMS_PdbInfoSlot *slot = syms_push_array(arena, SYMS_PdbInfoSlot, 1);
+      SYMS_QueuePush(result.first, result.last, slot);
+      slot->string = name;
+      slot->sn = sn;
+    }
+  }
+  result.auth_guid = auth_guid;
+  
+  SYMS_ProfEnd();
+  
+  return(result);
+}
+
+SYMS_API SYMS_PdbNamedStreamArray
+syms_pdb_named_stream_array(SYMS_PdbInfoTable *table){
+  SYMS_ProfBegin("syms_pdb_named_stream_array");
+  
+  // mapping "PdbDbiNamedStream" indexes to strings
+  struct StreamNameIndexPair{
+    SYMS_PdbNamedStream index;
+    SYMS_String8 name;
+  };
+  struct StreamNameIndexPair pairs[] = {
+    {SYMS_PdbNamedStream_HEADER_BLOCK, syms_str8_lit("/src/headerblock")},
+    {SYMS_PdbNamedStream_STRTABLE    , syms_str8_lit("/names")},
+    {SYMS_PdbNamedStream_LINK_INFO   , syms_str8_lit("/LinkInfo")},
+  };
+  
+  // build baked array of stream indices
+  SYMS_PdbNamedStreamArray result = {0};
+  struct StreamNameIndexPair *p = pairs;
+  for (SYMS_U64 i = 0; i < SYMS_ARRAY_SIZE(pairs); i += 1, p += 1){
+    SYMS_String8 name = p->name;
+    // slot from name
+    SYMS_PdbInfoSlot *matching_slot = 0;
+    for (SYMS_PdbInfoSlot *slot = table->first;
+         slot != 0;
+         slot = slot->next){
+      if (syms_string_match(name, slot->string, 0)){
+        matching_slot = slot;
+        break;
+      }
+    }
+    // store this slot's stream number into the array
+    if (matching_slot != 0){
+      result.sn[p->index] = matching_slot->sn;
+    }
+  }
+  
+  SYMS_ProfEnd();
+  
+  return(result);
+}
+
+SYMS_API SYMS_PdbStrtblAccel
+syms_pdb_dbi_parse_strtbl(SYMS_Arena *arena, SYMS_String8 data, SYMS_MsfAccel *msf, SYMS_MsfStreamNumber sn){
+  SYMS_ProfBegin("syms_pdb_dbi_parse_strtbl");
+  
+  // grab the strtable stream
+  SYMS_MsfRange range = syms_msf_range_from_sn(msf, sn);
+  
+  // read strtable's header
+  SYMS_PdbDbiStrTableHeader header = {0};
+  syms_msf_read_struct_in_range(data, msf, range, 0, &header);
+  
+  // read strtable's layout
+  SYMS_PdbStrtblAccel result = {0};
+  if (header.magic == 0xEFFEEFFE && header.version == 1){
+    SYMS_U32 strblock_size_off = sizeof(header);
+    SYMS_U32 strblock_size = 0;
+    syms_msf_read_struct_in_range(data, msf, range, strblock_size_off, &strblock_size);
+    
+    SYMS_U32 strblock_off = strblock_size_off + 4;
+    SYMS_U32 bucket_count_off = strblock_off + strblock_size;
+    SYMS_U32 bucket_count = 0;
+    syms_msf_read_struct_in_range(data, msf, range, bucket_count_off, &bucket_count);
+    
+    SYMS_U32 bucket_array_off = bucket_count_off + 4;
+    SYMS_U32 bucket_array_size = bucket_count*sizeof(SYMS_PdbStringIndex);
+    
+    if (syms_msf_bounds_check_in_range(range, bucket_array_off + bucket_array_size)){
+      result.bucket_count = bucket_count;
+      result.sn = sn;
+      result.strblock.min = range.off + strblock_off;
+      result.strblock.max = result.strblock.min + strblock_size;
+      result.buckets.min  = range.off + bucket_array_off;
+      result.buckets.max  = result.buckets.min + bucket_array_size;
+    }
+  }
+  
+  SYMS_ProfEnd();
+  
+  return(result);
+}
+
+
+// pdb specific api
+
+SYMS_API SYMS_MsfAccel*
+syms_pdb_msf_accel_from_dbg(SYMS_PdbDbgAccel *dbg){
+  SYMS_MsfAccel *result = dbg->msf;
+  return(result);
+}
+
+SYMS_API SYMS_String8
+syms_pdb_strtbl_string_from_off(SYMS_Arena *arena, SYMS_String8 data, SYMS_PdbDbgAccel *dbg, SYMS_U32 off){
+  // setup accel
+  SYMS_MsfAccel *msf = dbg->msf;
+  SYMS_PdbStrtblAccel *strtbl = &dbg->strtbl;
+  
+  // get string range
+  SYMS_MsfRange range = syms_msf_range_from_sn(msf, strtbl->sn);
+  SYMS_MsfRange str_range = syms_msf_sub_range_from_off_range(range, strtbl->strblock);
+  
+  // read string
+  SYMS_String8 result = syms_msf_read_zstring_in_range(arena, data, msf, str_range, off);
+  return(result);
+}
+
+SYMS_API SYMS_String8
+syms_pdb_strtbl_string_from_index(SYMS_Arena *arena, SYMS_String8 data, SYMS_PdbDbgAccel *dbg,
+                                  SYMS_PdbStringIndex n){
+  // setup accel
+  SYMS_MsfAccel *msf = dbg->msf;
+  SYMS_PdbStrtblAccel *strtbl = &dbg->strtbl;
+  
+  SYMS_String8 result = {0};
+  if (n < strtbl->bucket_count){
+    // get bucket range
+    SYMS_MsfRange range = syms_msf_range_from_sn(msf, strtbl->sn);
+    SYMS_MsfRange bkt_range = syms_msf_sub_range_from_off_range(range, strtbl->buckets);
+    
+    // get offset
+    SYMS_U32 offset_off = n*4;
+    SYMS_U32 offset = 0;
+    if (syms_msf_read_struct_in_range(data, msf, bkt_range, offset_off, &offset)){
+      
+      // get string range
+      SYMS_MsfRange str_range = syms_msf_sub_range_from_off_range(range, strtbl->strblock);
+      
+      // read string
+      result = syms_msf_read_zstring_in_range(arena, data, msf, str_range, offset);
+    }
+  }
+  
+  return(result);
+}
+
+
+// main api
+
+SYMS_API SYMS_PdbFileAccel*
+syms_pdb_file_accel_from_data(SYMS_Arena *arena, SYMS_String8 data){
+  SYMS_PdbFileAccel *result = (SYMS_PdbFileAccel*)&syms_format_nil;
+  SYMS_MsfAccel *msf = syms_msf_accel_from_data(arena, data);
+  if (msf != 0){
+    result = syms_push_array(arena, SYMS_PdbFileAccel, 1);
+    result->format = SYMS_FileFormat_PDB;
+    result->msf = msf;
+  }
+  return(result);
+}
+
+SYMS_API SYMS_PdbDbgAccel*
+syms_pdb_dbg_accel_from_file(SYMS_Arena *arena, SYMS_String8 data, SYMS_PdbFileAccel *file){
+  // setup accel
+  SYMS_MsfAccel *msf = file->msf;
+  
+  // setup result
+  SYMS_PdbDbgAccel *result = (SYMS_PdbDbgAccel*)&syms_format_nil;
+  
+  // parse dbi header
+  SYMS_PdbDbiAccel dbi = syms_pdb_dbi_from_msf(data, msf);
+  if (dbi.valid){
+    // copy msf if it's not a prior on this arena.
+    // TODO(allen): re-apply deep copy optimizations using handle checks
+    SYMS_MsfAccel *dbg_msf = syms_msf_deep_copy(arena, msf);
+    
+    // parse info table
+    SYMS_ArenaTemp scratch = syms_get_scratch(&arena, 1);
+    SYMS_PdbInfoTable info_table = syms_pdb_parse_info(scratch.arena, data, msf);
+    
+    // get gsi range
+    SYMS_MsfRange gsi_range = syms_msf_range_from_sn(msf, dbi.gsi_sn);
+    SYMS_MsfRange psi_range = syms_msf_range_from_sn(msf, dbi.psi_sn);
+    SYMS_MsfRange gsi_part_psi_range = syms_pdb_gsi_part_from_psi_range(psi_range);
+    
+    // fill accelerator
+    result = syms_push_array(arena, SYMS_PdbDbgAccel, 1);
+    result->format = SYMS_FileFormat_PDB;
+    result->msf = dbg_msf;
+    result->dbi = dbi;
+    result->named = syms_pdb_named_stream_array(&info_table);
+    result->strtbl = syms_pdb_dbi_parse_strtbl(arena, data, msf, result->named.sn[SYMS_PdbNamedStream_STRTABLE]);
+    result->tpi = syms_pdb_tpi_accel_from_sn(arena, data, msf, SYMS_PdbFixedStream_TPI);
+    result->ipi = syms_pdb_tpi_accel_from_sn(arena, data, msf, SYMS_PdbFixedStream_IPI);
+    result->gsi = syms_pdb_gsi_accel_from_range(arena, data, msf, gsi_range);
+    result->psi = syms_pdb_gsi_accel_from_range(arena, data, msf, gsi_part_psi_range);
+    result->auth_guid = info_table.auth_guid;
+    
+    // release scratch
+    syms_release_scratch(scratch);
+  }
+  return(result);
+}
+
+
+////////////////////////////////
+//~ NOTE(allen): PDB Architecture
+
+SYMS_API SYMS_Arch
+syms_pdb_arch_from_dbg(SYMS_PdbDbgAccel *dbg){
+  SYMS_Arch result = syms_arch_from_coff_machine_type(dbg->dbi.machine_type);
+  return(result);
+}
+
+
+////////////////////////////////
+//~ NOTE(allen): PDB Match Keys
+
+SYMS_API SYMS_ExtMatchKey
+syms_pdb_ext_match_key_from_dbg(SYMS_String8 data, SYMS_PdbDbgAccel *dbg){
+  SYMS_ExtMatchKey result = {0};
+  syms_memmove(&result, &dbg->auth_guid, sizeof(dbg->auth_guid));
+  return(result);
+}
+
+
+////////////////////////////////
+//~ NOTE(allen): PDB Sections
+
+// pdb specific
+
+SYMS_API SYMS_CoffSection
+syms_pdb_coff_section(SYMS_String8 data, SYMS_PdbDbgAccel *dbg, SYMS_U64 n){
+  SYMS_CoffSection result = {0};
+  if (1 <= n){
+    SYMS_MsfAccel *msf = dbg->msf;
+    SYMS_MsfRange range = syms_pdb_dbi_stream(msf, &dbg->dbi, SYMS_PdbDbiStream_SECTION_HEADER);
+    SYMS_U64 off = (n - 1)*sizeof(SYMS_CoffSection);
+    syms_msf_read_struct_in_range(data, msf, range, off, &result);
+  }
+  return(result);
+}
+
+// main api
+
+SYMS_API SYMS_SecInfoArray
+syms_pdb_sec_info_array_from_dbg(SYMS_Arena *arena, SYMS_String8 data, SYMS_PdbDbgAccel *dbg){
+  SYMS_ArenaTemp scratch = syms_get_scratch(&arena, 1);
+  
+  // setup accels
+  SYMS_MsfAccel *msf = dbg->msf;
+  SYMS_PdbDbiAccel *dbi = &dbg->dbi;
+  
+  // size of section header / size of coff section
+  SYMS_MsfRange range = syms_pdb_dbi_stream(msf, dbi, SYMS_PdbDbiStream_SECTION_HEADER);
+  SYMS_U64 count = range.size/sizeof(SYMS_CoffSection);
+  
+  // grab coff array
+  SYMS_CoffSection *coff_secs = syms_push_array(scratch.arena, SYMS_CoffSection, count);
+  syms_msf_read_in_range(data, msf, range, 0, count*sizeof(*coff_secs), coff_secs);
+  
+  // fill array
+  SYMS_SecInfoArray result = {0};
+  result.count = count;
+  result.sec_info = syms_push_array(arena, SYMS_SecInfo, count);
+  
+  SYMS_SecInfo *sec_info = result.sec_info;
+  SYMS_CoffSection *coff_sec = coff_secs;
+  
+  for (SYMS_U64 i = 0; i < result.count; i += 1, sec_info += 1, coff_sec += 1){
+    // extract name
+    SYMS_U8 *name = coff_sec->name;
+    SYMS_U8 *name_ptr = name;
+    SYMS_U8 *name_opl = name + 8;
+    for (;name_ptr < name_opl && *name_ptr != 0; name_ptr += 1);
+    
+    // fill sec info
+    sec_info->name = syms_push_string_copy(arena, syms_str8_range(name, name_ptr));
+    sec_info->vrange.min = coff_sec->virt_off;
+    sec_info->vrange.max = coff_sec->virt_off + coff_sec->virt_size;
+    sec_info->frange.min = coff_sec->file_off;
+    sec_info->frange.max = coff_sec->file_off + coff_sec->file_size;
+  }
+  
+  syms_release_scratch(scratch);
+  
+  return(result);
+}
+
+////////////////////////////////
+//~ NOTE(allen): PDB Compilation Units
+
+// pdb helpers
+
+SYMS_API SYMS_PdbCompUnit*
+syms_pdb_comp_unit_from_id(SYMS_PdbUnitSetAccel *accel, SYMS_UnitID id){
+  SYMS_PdbCompUnit *result = 0;
+  if (SYMS_PdbPseudoUnit_FIRST_COMP_UNIT <= id){
+    SYMS_U64 index = id - SYMS_PdbPseudoUnit_FIRST_COMP_UNIT;
+    if (index < accel->comp_count){
+      result = accel->comp_units[index];
+    }
+  }
+  return(result);
+}
+
+SYMS_API SYMS_MsfRange
+syms_pdb_msf_range_from_comp_unit(SYMS_PdbCompUnit *unit, SYMS_PdbCompUnitRange n){
+  SYMS_MsfRange result = {0};
+  result.sn = unit->sn;
+  result.off = unit->range_off[n];
+  result.size = unit->range_off[n + 1] - result.off;
+  return(result);
+}
+
+// main api
+
+SYMS_API SYMS_PdbUnitSetAccel*
+syms_pdb_unit_set_accel_from_dbg(SYMS_Arena *arena, SYMS_String8 data, SYMS_PdbDbgAccel *dbg){
+  // setup accels
+  SYMS_MsfAccel *msf = dbg->msf;
+  SYMS_PdbDbiAccel *dbi = &dbg->dbi;
+  
+  // parse list of modules
+  SYMS_PdbCompUnitNode *first = 0;
+  SYMS_PdbCompUnitNode *last = 0;
+  SYMS_U64 count = 0;
+  
+  {
+    // grab module info range
+    SYMS_MsfRange range = syms_pdb_dbi_sub_range(dbi, SYMS_PdbDbiRange_ModuleInfo);
+    
+    // parse loop
+    SYMS_U64 off = 0;
+    for (;;){
+      // exit condition - past end
+      if (off >= range.size){
+        break;
+      }
+      
+      // extract comp unit info
+      SYMS_B32 got_comp_unit = syms_false;
+      SYMS_PdbCompUnit comp_unit = {0};
+      SYMS_U32 next_comp_unit_off = 0;
+      
+      SYMS_PdbDbiCompUnitHeader header = {0};
+      if (syms_msf_read_struct_in_range(data, msf, range, off, &header)){
+        got_comp_unit = syms_true;
+        
+        // extract names from after header
+        SYMS_U32 name_off = off + sizeof(header);
+        SYMS_String8 name = syms_msf_read_zstring_in_range(arena, data, msf, range, name_off);
+        
+        SYMS_U32 name2_off = name_off + name.size + 1;
+        SYMS_String8 name2 = syms_msf_read_zstring_in_range(arena, data, msf, range, name2_off);
+        
+        SYMS_U32 after_name2_off = name2_off + name2.size + 1;
+        next_comp_unit_off = (after_name2_off + 3)&~3;
+        
+        // get the size of this mod's stream
+        SYMS_MsfStreamInfo stream_info = syms_msf_stream_info_from_sn(msf, header.sn);
+        SYMS_U64 stream_size = stream_info.size;
+        
+        // organize size information
+        SYMS_U64 range_size[SYMS_PdbCompUnitRange_COUNT];
+        range_size[SYMS_PdbCompUnitRange_Symbols] = header.symbol_bytes;
+        range_size[SYMS_PdbCompUnitRange_C11] = header.c11_lines_size;
+        range_size[SYMS_PdbCompUnitRange_C13] = header.c13_lines_size;
+        
+        // fill comp unit range offs
+        {
+          SYMS_U64 cursor = 0;
+          SYMS_U64 i = 0;
+          for (; i < (SYMS_U64)(SYMS_PdbCompUnitRange_COUNT); i += 1){
+            comp_unit.range_off[i] = cursor;
+            cursor += range_size[i];
+            cursor = SYMS_ClampTop(cursor, stream_size);
+          }
+          //~ NOTE(allen): one last value past the end so that off[i + 1] - off[i] can get us the sizes back.
+          comp_unit.range_off[i] = cursor;
+        }
+        
+        // adjust symbols range off up by 4 to skip annoying signature
+        comp_unit.range_off[SYMS_PdbCompUnitRange_Symbols] += 4;
+        
+        // fill comp unit direct fields
+        comp_unit.sn = header.sn;
+        comp_unit.src_file = header.src_file;
+        comp_unit.pdb_file = header.pdb_file;
+        
+        comp_unit.module_name = name;
+        comp_unit.obj_name = name2;
+      }
+      
+      // exit condition - failed read
+      if (!got_comp_unit){
+        break;
+      }
+      
+      // save to list
+      SYMS_PdbCompUnitNode *node = syms_push_array(arena, SYMS_PdbCompUnitNode, 1);
+      SYMS_QueuePush(first, last, node);
+      count += 1;
+      syms_memmove(&node->comp_unit, &comp_unit, sizeof(comp_unit));
+      
+      // increment offset
+      off = next_comp_unit_off;
+    }
+  }
+  
+  // linearize for quick index
+  SYMS_PdbCompUnit **comp_units = syms_push_array(arena, SYMS_PdbCompUnit*, count);
+  {
+    SYMS_PdbCompUnit **comp_unit_ptr = comp_units;
+    for (SYMS_PdbCompUnitNode *node = first;
+         node != 0;
+         node = node->next, comp_unit_ptr += 1){
+      *comp_unit_ptr = &node->comp_unit;
+    }
+  }
+  
+  // fill result
+  SYMS_PdbUnitSetAccel *result = syms_push_array(arena, SYMS_PdbUnitSetAccel, 1);
+  result->format = SYMS_FileFormat_PDB;
+  result->comp_units = comp_units;
+  result->comp_count = count;
+  
+  return(result);
+}
+
+SYMS_API SYMS_U64
+syms_pdb_unit_count_from_set(SYMS_PdbUnitSetAccel *unit_set){
+  SYMS_U64 result = unit_set->comp_count + SYMS_PdbPseudoUnit_COUNT;
+  return(result);
+}
+
+SYMS_API SYMS_UnitInfo
+syms_pdb_unit_info_from_uid(SYMS_PdbUnitSetAccel *unit_set, SYMS_UnitID uid){
+  SYMS_UnitInfo result = {0};
+  
+  // TODO(allen): language
+  
+  // try symbols pseudo unit
+  if (uid == SYMS_PdbPseudoUnit_SYM){
+    result.features = (SYMS_UnitFeature_StaticVariables|
+                       SYMS_UnitFeature_ExternVariables|
+                       SYMS_UnitFeature_FunctionStubs);
+  }
+  
+  // try types pseudo unit
+  else if (uid == SYMS_PdbPseudoUnit_TPI){
+    result.features = SYMS_UnitFeature_Types;
+  }
+  
+  // try compilation unit
+  else{
+    SYMS_PdbCompUnit *unit = syms_pdb_comp_unit_from_id(unit_set, uid);
+    if (unit != 0){
+      result.features = (SYMS_UnitFeature_CompilationUnit|
+                         SYMS_UnitFeature_StaticVariables|
+                         SYMS_UnitFeature_Functions);
+    }
+  }
+  
+  return(result);
+}
+
+SYMS_API SYMS_UnitNames
+syms_pdb_unit_names_from_uid(SYMS_Arena *arena, SYMS_PdbUnitSetAccel *unit_set, SYMS_UnitID uid){
+  SYMS_UnitNames result = {0};
+  // try compilation unit
+  SYMS_PdbCompUnit *unit = syms_pdb_comp_unit_from_id(unit_set, uid);
+  if (unit != 0){
+    result.object_file  = syms_push_string_copy(arena, unit->module_name);
+    result.archive_file = syms_push_string_copy(arena, unit->obj_name);
+  }
+  return(result);
+}
+
+SYMS_API SYMS_UnitRangeArray
+syms_pdb_unit_ranges_from_set(SYMS_Arena *arena, SYMS_String8 data, SYMS_PdbDbgAccel *dbg,
+                              SYMS_PdbUnitSetAccel *unit_set){
+  //- setup accels
+  SYMS_MsfAccel *msf = dbg->msf;
+  SYMS_PdbDbiAccel *dbi = &dbg->dbi;
+  
+  //- parse contributions
+  SYMS_U64 range_count = 0;
+  SYMS_UnitRange *ranges = 0;
+  {
+    // grab module info range
+    SYMS_MsfRange range = syms_pdb_dbi_sub_range(dbi, SYMS_PdbDbiRange_SecCon);
+    
+    // read version
+    SYMS_PdbDbiSectionContribVersion version = 0;
+    syms_msf_read_struct_in_range(data, msf, range, 0, &version);
+    
+    // read array off/size info
+    SYMS_U32 item_size = 0;
+    SYMS_U32 array_off = 0;
+    switch (version){
+      default:
+      {
+        // TODO(allen): do we have a test case for this?
+        version = 0;
+        item_size = sizeof(SYMS_PdbDbiSectionContrib40);
+      }break;
+      case SYMS_PdbDbiSectionContribVersion_1:
+      {
+        item_size = sizeof(SYMS_PdbDbiSectionContrib);
+        array_off = sizeof(version);
+      }break;
+      case SYMS_PdbDbiSectionContribVersion_2:
+      {
+        item_size = sizeof(SYMS_PdbDbiSectionContrib2);
+        array_off = sizeof(version);
+      }break;
+    }
+    
+    // allocate ranges array
+    SYMS_U64 max_count = (range.size - array_off)/item_size;
+    ranges = syms_push_array(arena, SYMS_UnitRange, max_count);
+    
+    // fill array
+    {
+      SYMS_UnitRange *unit_range = ranges;
+      SYMS_UnitRange *opl = ranges + max_count;
+      for (SYMS_U32 off = array_off;
+           unit_range < opl;
+           off += item_size){
+        SYMS_PdbDbiSectionContrib40 sc = {0};
+        if (!syms_msf_read_struct_in_range(data, msf, range, off, &sc)){
+          break;
+        }
+        
+        if (sc.size > 0){
+          // calculate range
+          SYMS_CoffSection section = syms_pdb_coff_section(data, dbg, sc.sec);
+          SYMS_U64 min = section.virt_off + sc.sec_off;
+          
+          // fill unit range
+          unit_range->vrange.min = min;
+          unit_range->vrange.max = min + sc.size;
+          unit_range->uid = sc.mod + SYMS_PdbPseudoUnit_FIRST_COMP_UNIT;
+          unit_range += 1;
+        }
+      }
+      
+      range_count = (SYMS_U64)(unit_range - ranges);
+      
+      syms_arena_put_back(arena, sizeof(SYMS_UnitRange)*(max_count - range_count));
+    }
+  }
+  
+  //- assemble result
+  SYMS_UnitRangeArray result = {ranges, range_count};
+  return(result);
+}
+
+////////////////////////////////
+//~ NOTE(allen): PDB Symbol Parsing
+
+// cv parsing helpers
+
+SYMS_API void
+syms_pdb_leaf_debug_helper(SYMS_CvLeaf kind){
+  switch (kind){
+    case SYMS_CvLeaf_MODIFIER_16t:break;
+    case SYMS_CvLeaf_POINTER_16t:break;
+    case SYMS_CvLeaf_ARRAY_16t:break;
+    case SYMS_CvLeaf_CLASS_16t:break;
+    case SYMS_CvLeaf_STRUCTURE_16t:break;
+    case SYMS_CvLeaf_UNION_16t:break;
+    case SYMS_CvLeaf_ENUM_16t:break;
+    case SYMS_CvLeaf_PROCEDURE_16t:break;
+    case SYMS_CvLeaf_MFUNCTION_16t:break;
+    case SYMS_CvLeaf_VTSHAPE:break;
+    case SYMS_CvLeaf_COBOL0_16t:break;
+    case SYMS_CvLeaf_COBOL1:break;
+    case SYMS_CvLeaf_BARRAY_16t:break;
+    case SYMS_CvLeaf_LABEL:break;
+    case SYMS_CvLeaf_NULL:break;
+    case SYMS_CvLeaf_NOTTRAN:break;
+    case SYMS_CvLeaf_DIMARRAY_16t:break;
+    case SYMS_CvLeaf_VFTPATH_16t:break;
+    case SYMS_CvLeaf_PRECOMP_16t:break;
+    case SYMS_CvLeaf_ENDPRECOMP:break;
+    case SYMS_CvLeaf_OEM_16t:break;
+    case SYMS_CvLeaf_TYPESERVER_ST:break;
+    case SYMS_CvLeaf_SKIP_16t:break;
+    case SYMS_CvLeaf_ARGLIST_16t:break;
+    case SYMS_CvLeaf_DEFARG_16t:break;
+    case SYMS_CvLeaf_LIST:break;
+    case SYMS_CvLeaf_FIELDLIST_16t:break;
+    case SYMS_CvLeaf_DERIVED_16t:break;
+    case SYMS_CvLeaf_BITFIELD_16t:break;
+    case SYMS_CvLeaf_METHODLIST_16t:break;
+    case SYMS_CvLeaf_DIMCONU_16t:break;
+    case SYMS_CvLeaf_DIMCONLU_16t:break;
+    case SYMS_CvLeaf_DIMVARU_16t:break;
+    case SYMS_CvLeaf_DIMVARLU_16t:break;
+    case SYMS_CvLeaf_REFSYM:break;
+    case SYMS_CvLeaf_BCLASS_16t:break;
+    case SYMS_CvLeaf_VBCLASS_16t:break;
+    case SYMS_CvLeaf_IVBCLASS_16t:break;
+    case SYMS_CvLeaf_ENUMERATE_ST:break;
+    case SYMS_CvLeaf_FRIENDFCN_16t:break;
+    case SYMS_CvLeaf_INDEX_16t:break;
+    case SYMS_CvLeaf_MEMBER_16t:break;
+    case SYMS_CvLeaf_STMEMBER_16t:break;
+    case SYMS_CvLeaf_METHOD_16t:break;
+    case SYMS_CvLeaf_NESTTYPE_16t:break;
+    case SYMS_CvLeaf_VFUNCTAB_16t:break;
+    case SYMS_CvLeaf_FRIENDCLS_16t:break;
+    case SYMS_CvLeaf_ONEMETHOD_16t:break;
+    case SYMS_CvLeaf_VFUNCOFF_16t:break;
+    case SYMS_CvLeaf_TI16_MAX:break;
+    case SYMS_CvLeaf_MODIFIER:break;
+    case SYMS_CvLeaf_POINTER:break;
+    case SYMS_CvLeaf_ARRAY_ST:break;
+    case SYMS_CvLeaf_CLASS_ST:break;
+    case SYMS_CvLeaf_STRUCTURE_ST:break;
+    case SYMS_CvLeaf_UNION_ST:break;
+    case SYMS_CvLeaf_ENUM_ST:break;
+    case SYMS_CvLeaf_PROCEDURE:break;
+    case SYMS_CvLeaf_MFUNCTION:break;
+    case SYMS_CvLeaf_COBOL0:break;
+    case SYMS_CvLeaf_BARRAY:break;
+    case SYMS_CvLeaf_DIMARRAY_ST:break;
+    case SYMS_CvLeaf_VFTPATH:break;
+    case SYMS_CvLeaf_PRECOMP_ST:break;
+    case SYMS_CvLeaf_OEM:break;
+    case SYMS_CvLeaf_ALIAS_ST:break;
+    case SYMS_CvLeaf_OEM2:break;
+    case SYMS_CvLeaf_SKIP:break;
+    case SYMS_CvLeaf_ARGLIST:break;
+    case SYMS_CvLeaf_DEFARG_ST:break;
+    case SYMS_CvLeaf_FIELDLIST:break;
+    case SYMS_CvLeaf_DERIVED:break;
+    case SYMS_CvLeaf_BITFIELD:break;
+    case SYMS_CvLeaf_METHODLIST:break;
+    case SYMS_CvLeaf_DIMCONU:break;
+    case SYMS_CvLeaf_DIMCONLU:break;
+    case SYMS_CvLeaf_DIMVARU:break;
+    case SYMS_CvLeaf_DIMVARLU:break;
+    case SYMS_CvLeaf_BCLASS:break;
+    case SYMS_CvLeaf_VBCLASS:break;
+    case SYMS_CvLeaf_IVBCLASS:break;
+    case SYMS_CvLeaf_FRIENDFCN_ST:break;
+    case SYMS_CvLeaf_INDEX:break;
+    case SYMS_CvLeaf_MEMBER_ST:break;
+    case SYMS_CvLeaf_STMEMBER_ST:break;
+    case SYMS_CvLeaf_METHOD_ST:break;
+    case SYMS_CvLeaf_NESTTYPE_ST:break;
+    case SYMS_CvLeaf_VFUNCTAB:break;
+    case SYMS_CvLeaf_FRIENDCLS:break;
+    case SYMS_CvLeaf_ONEMETHOD_ST:break;
+    case SYMS_CvLeaf_VFUNCOFF:break;
+    case SYMS_CvLeaf_NESTTYPEEX_ST:break;
+    case SYMS_CvLeaf_MEMBERMODIFY_ST:break;
+    case SYMS_CvLeaf_MANAGED_ST:break;
+    case SYMS_CvLeaf_ST_MAX:break;
+    case SYMS_CvLeaf_TYPESERVER:break;
+    case SYMS_CvLeaf_ENUMERATE:break;
+    case SYMS_CvLeaf_ARRAY:break;
+    case SYMS_CvLeaf_CLASS:break;
+    case SYMS_CvLeaf_STRUCTURE:break;
+    case SYMS_CvLeaf_UNION:break;
+    case SYMS_CvLeaf_ENUM:break;
+    case SYMS_CvLeaf_DIMARRAY:break;
+    case SYMS_CvLeaf_PRECOMP:break;
+    case SYMS_CvLeaf_ALIAS:break;
+    case SYMS_CvLeaf_DEFARG:break;
+    case SYMS_CvLeaf_FRIENDFCN:break;
+    case SYMS_CvLeaf_MEMBER:break;
+    case SYMS_CvLeaf_STMEMBER:break;
+    case SYMS_CvLeaf_METHOD:break;
+    case SYMS_CvLeaf_NESTTYPE:break;
+    case SYMS_CvLeaf_ONEMETHOD:break;
+    case SYMS_CvLeaf_NESTTYPEEX:break;
+    case SYMS_CvLeaf_MEMBERMODIFY:break;
+    case SYMS_CvLeaf_MANAGED:break;
+    case SYMS_CvLeaf_TYPESERVER2:break;
+    case SYMS_CvLeaf_STRIDED_ARRAY:break;
+    case SYMS_CvLeaf_HLSL:break;
+    case SYMS_CvLeaf_MODIFIER_EX:break;
+    case SYMS_CvLeaf_INTERFACE:break;
+    case SYMS_CvLeaf_BINTERFACE:break;
+    case SYMS_CvLeaf_VECTOR:break;
+    case SYMS_CvLeaf_MATRIX:break;
+    case SYMS_CvLeaf_VFTABLE:break;
+    case SYMS_CvLeaf_FUNC_ID:break;
+    case SYMS_CvLeaf_MFUNC_ID:break;
+    case SYMS_CvLeaf_BUILDINFO:break;
+    case SYMS_CvLeaf_SUBSTR_LIST:break;
+    case SYMS_CvLeaf_STRING_ID:break;
+    case SYMS_CvLeaf_UDT_SRC_LINE:break;
+    case SYMS_CvLeaf_UDT_MOD_SRC_LINE:break;
+    case SYMS_CvLeaf_CLASSPTR:break;
+    case SYMS_CvLeaf_CLASSPTR2:break;
+    case SYMS_CvLeaf_CHAR:break;
+    case SYMS_CvLeaf_SHORT:break;
+    case SYMS_CvLeaf_USHORT:break;
+    case SYMS_CvLeaf_LONG:break;
+    case SYMS_CvLeaf_ULONG:break;
+    case SYMS_CvLeaf_FLOAT32:break;
+    case SYMS_CvLeaf_FLOAT64:break;
+    case SYMS_CvLeaf_FLOAT80:break;
+    case SYMS_CvLeaf_FLOAT128:break;
+    case SYMS_CvLeaf_QUADWORD:break;
+    case SYMS_CvLeaf_UQUADWORD:break;
+    case SYMS_CvLeaf_FLOAT48:break;
+    case SYMS_CvLeaf_COMPLEX32:break;
+    case SYMS_CvLeaf_COMPLEX64:break;
+    case SYMS_CvLeaf_COMPLEX80:break;
+    case SYMS_CvLeaf_COMPLEX128:break;
+    case SYMS_CvLeaf_VARSTRING:break;
+    case SYMS_CvLeaf_OCTWORD:break;
+    case SYMS_CvLeaf_UOCTWORD:break;
+    case SYMS_CvLeaf_DECIMAL:break;
+    case SYMS_CvLeaf_DATE:break;
+    case SYMS_CvLeaf_UTF8STRING:break;
+    case SYMS_CvLeaf_FLOAT16:break;
+    case SYMS_CvLeaf_PAD0:break;
+    case SYMS_CvLeaf_PAD1:break;
+    case SYMS_CvLeaf_PAD2:break;
+    case SYMS_CvLeaf_PAD3:break;
+    case SYMS_CvLeaf_PAD4:break;
+    case SYMS_CvLeaf_PAD5:break;
+    case SYMS_CvLeaf_PAD6:break;
+    case SYMS_CvLeaf_PAD7:break;
+    case SYMS_CvLeaf_PAD8:break;
+    case SYMS_CvLeaf_PAD9:break;
+    case SYMS_CvLeaf_PAD10:break;
+    case SYMS_CvLeaf_PAD11:break;
+    case SYMS_CvLeaf_PAD12:break;
+    case SYMS_CvLeaf_PAD13:break;
+    case SYMS_CvLeaf_PAD14:break;
+    case SYMS_CvLeaf_PAD15:break;
+  }
+}
+
+SYMS_API void
+syms_pdb_sym_debug_helper(SYMS_CvSymKind kind){
+  switch (kind){
+    case SYMS_CvSymKind_NULL:break;
+    case SYMS_CvSymKind_COMPILE:break;
+    case SYMS_CvSymKind_REGISTER_16t:break;
+    case SYMS_CvSymKind_CONSTANT_16t:break;
+    case SYMS_CvSymKind_UDT_16t:break;
+    case SYMS_CvSymKind_SSEARCH:break;
+    case SYMS_CvSymKind_END:break;
+    case SYMS_CvSymKind_SKIP:break;
+    case SYMS_CvSymKind_CVRESERVE:break;
+    case SYMS_CvSymKind_OBJNAME_ST:break;
+    case SYMS_CvSymKind_ENDARG:break;
+    case SYMS_CvSymKind_COBOLUDT_16t:break;
+    case SYMS_CvSymKind_MANYREG_16t:break;
+    case SYMS_CvSymKind_RETURN:break;
+    case SYMS_CvSymKind_ENTRYTHIS:break;
+    case SYMS_CvSymKind_BPREL16:break;
+    case SYMS_CvSymKind_LDATA16:break;
+    case SYMS_CvSymKind_GDATA16:break;
+    case SYMS_CvSymKind_PUB16:break;
+    case SYMS_CvSymKind_LPROC16:break;
+    case SYMS_CvSymKind_GPROC16:break;
+    case SYMS_CvSymKind_THUNK16:break;
+    case SYMS_CvSymKind_BLOCK16:break;
+    case SYMS_CvSymKind_WITH16:break;
+    case SYMS_CvSymKind_LABEL16:break;
+    case SYMS_CvSymKind_CEXMODEL16:break;
+    case SYMS_CvSymKind_VFTABLE16:break;
+    case SYMS_CvSymKind_REGREL16:break;
+    case SYMS_CvSymKind_BPREL32_16t:break;
+    case SYMS_CvSymKind_LDATA32_16t:break;
+    case SYMS_CvSymKind_GDATA32_16t:break;
+    case SYMS_CvSymKind_PUB32_16t:break;
+    case SYMS_CvSymKind_LPROC32_16t:break;
+    case SYMS_CvSymKind_GPROC32_16t:break;
+    case SYMS_CvSymKind_THUNK32_ST:break;
+    case SYMS_CvSymKind_BLOCK32_ST:break;
+    case SYMS_CvSymKind_WITH32_ST:break;
+    case SYMS_CvSymKind_LABEL32_ST:break;
+    case SYMS_CvSymKind_CEXMODEL32:break;
+    case SYMS_CvSymKind_VFTABLE32_16t:break;
+    case SYMS_CvSymKind_REGREL32_16t:break;
+    case SYMS_CvSymKind_LTHREAD32_16t:break;
+    case SYMS_CvSymKind_GTHREAD32_16t:break;
+    case SYMS_CvSymKind_SLINK32:break;
+    case SYMS_CvSymKind_LPROCMIPS_16t:break;
+    case SYMS_CvSymKind_GPROCMIPS_16t:break;
+    case SYMS_CvSymKind_PROCREF_ST:break;
+    case SYMS_CvSymKind_DATAREF_ST:break;
+    case SYMS_CvSymKind_ALIGN:break;
+    case SYMS_CvSymKind_LPROCREF_ST:break;
+    case SYMS_CvSymKind_OEM:break;
+    case SYMS_CvSymKind_TI16_MAX:break;
+    case SYMS_CvSymKind_CONSTANT_ST:break;
+    case SYMS_CvSymKind_UDT_ST:break;
+    case SYMS_CvSymKind_COBOLUDT_ST:break;
+    case SYMS_CvSymKind_MANYREG_ST:break;
+    case SYMS_CvSymKind_BPREL32_ST:break;
+    case SYMS_CvSymKind_LDATA32_ST:break;
+    case SYMS_CvSymKind_GDATA32_ST:break;
+    case SYMS_CvSymKind_PUB32_ST:break;
+    case SYMS_CvSymKind_LPROC32_ST:break;
+    case SYMS_CvSymKind_GPROC32_ST:break;
+    case SYMS_CvSymKind_VFTABLE32:break;
+    case SYMS_CvSymKind_REGREL32_ST:break;
+    case SYMS_CvSymKind_LTHREAD32_ST:break;
+    case SYMS_CvSymKind_GTHREAD32_ST:break;
+    case SYMS_CvSymKind_LPROCMIPS_ST:break;
+    case SYMS_CvSymKind_GPROCMIPS_ST:break;
+    case SYMS_CvSymKind_FRAMEPROC:break;
+    case SYMS_CvSymKind_COMPILE2_ST:break;
+    case SYMS_CvSymKind_MANYREG2_ST:break;
+    case SYMS_CvSymKind_LPROCIA64_ST:break;
+    case SYMS_CvSymKind_GPROCIA64_ST:break;
+    case SYMS_CvSymKind_LOCALSLOT_ST:break;
+    case SYMS_CvSymKind_PARAMSLOT_ST:break;
+    case SYMS_CvSymKind_GMANPROC_ST:break;
+    case SYMS_CvSymKind_LMANPROC_ST:break;
+    case SYMS_CvSymKind_RESERVED1:break;
+    case SYMS_CvSymKind_RESERVED2:break;
+    case SYMS_CvSymKind_RESERVED3:break;
+    case SYMS_CvSymKind_RESERVED4:break;
+    case SYMS_CvSymKind_LMANDATA_ST:break;
+    case SYMS_CvSymKind_GMANDATA_ST:break;
+    case SYMS_CvSymKind_MANFRAMEREL_ST:break;
+    case SYMS_CvSymKind_MANREGISTER_ST:break;
+    case SYMS_CvSymKind_MANSLOT_ST:break;
+    case SYMS_CvSymKind_MANMANYREG_ST:break;
+    case SYMS_CvSymKind_MANREGREL_ST:break;
+    case SYMS_CvSymKind_MANMANYREG2_ST:break;
+    case SYMS_CvSymKind_MANTYPREF:break;
+    case SYMS_CvSymKind_UNAMESPACE_ST:break;
+    case SYMS_CvSymKind_ST_MAX:break;
+    case SYMS_CvSymKind_OBJNAME:break;
+    case SYMS_CvSymKind_THUNK32:break;
+    case SYMS_CvSymKind_BLOCK32:break;
+    case SYMS_CvSymKind_WITH32:break;
+    case SYMS_CvSymKind_LABEL32:break;
+    case SYMS_CvSymKind_REGISTER:break;
+    case SYMS_CvSymKind_CONSTANT:break;
+    case SYMS_CvSymKind_UDT:break;
+    case SYMS_CvSymKind_COBOLUDT:break;
+    case SYMS_CvSymKind_MANYREG:break;
+    case SYMS_CvSymKind_BPREL32:break;
+    case SYMS_CvSymKind_LDATA32:break;
+    case SYMS_CvSymKind_GDATA32:break;
+    case SYMS_CvSymKind_PUB32:break;
+    case SYMS_CvSymKind_LPROC32:break;
+    case SYMS_CvSymKind_GPROC32:break;
+    case SYMS_CvSymKind_REGREL32:break;
+    case SYMS_CvSymKind_LTHREAD32:break;
+    case SYMS_CvSymKind_GTHREAD32:break;
+    case SYMS_CvSymKind_LPROCMIPS:break;
+    case SYMS_CvSymKind_GPROCMIPS:break;
+    case SYMS_CvSymKind_COMPILE2:break;
+    case SYMS_CvSymKind_MANYREG2:break;
+    case SYMS_CvSymKind_LPROCIA64:break;
+    case SYMS_CvSymKind_GPROCIA64:break;
+    case SYMS_CvSymKind_LOCALSLOT:break;
+    case SYMS_CvSymKind_PARAMSLOT:break;
+    case SYMS_CvSymKind_LMANDATA:break;
+    case SYMS_CvSymKind_GMANDATA:break;
+    case SYMS_CvSymKind_MANFRAMEREL:break;
+    case SYMS_CvSymKind_MANREGISTER:break;
+    case SYMS_CvSymKind_MANSLOT:break;
+    case SYMS_CvSymKind_MANMANYREG:break;
+    case SYMS_CvSymKind_MANREGREL:break;
+    case SYMS_CvSymKind_MANMANYREG2:break;
+    case SYMS_CvSymKind_UNAMESPACE:break;
+    case SYMS_CvSymKind_PROCREF:break;
+    case SYMS_CvSymKind_DATAREF:break;
+    case SYMS_CvSymKind_LPROCREF:break;
+    case SYMS_CvSymKind_ANNOTATIONREF:break;
+    case SYMS_CvSymKind_TOKENREF:break;
+    case SYMS_CvSymKind_GMANPROC:break;
+    case SYMS_CvSymKind_LMANPROC:break;
+    case SYMS_CvSymKind_TRAMPOLINE:break;
+    case SYMS_CvSymKind_MANCONSTANT:break;
+    case SYMS_CvSymKind_ATTR_FRAMEREL:break;
+    case SYMS_CvSymKind_ATTR_REGISTER:break;
+    case SYMS_CvSymKind_ATTR_REGREL:break;
+    case SYMS_CvSymKind_ATTR_MANYREG:break;
+    case SYMS_CvSymKind_SEPCODE:break;
+    case SYMS_CvSymKind_DEFRANGE_2005:break;
+    case SYMS_CvSymKind_DEFRANGE2_2005:break;
+    case SYMS_CvSymKind_SECTION:break;
+    case SYMS_CvSymKind_COFFGROUP:break;
+    case SYMS_CvSymKind_EXPORT:break;
+    case SYMS_CvSymKind_CALLSITEINFO:break;
+    case SYMS_CvSymKind_FRAMECOOKIE:break;
+    case SYMS_CvSymKind_DISCARDED:break;
+    case SYMS_CvSymKind_COMPILE3:break;
+    case SYMS_CvSymKind_ENVBLOCK:break;
+    case SYMS_CvSymKind_LOCAL:break;
+    case SYMS_CvSymKind_DEFRANGE:break;
+    case SYMS_CvSymKind_DEFRANGE_SUBFIELD:break;
+    case SYMS_CvSymKind_DEFRANGE_REGISTER:break;
+    case SYMS_CvSymKind_DEFRANGE_FRAMEPOINTER_REL:break;
+    case SYMS_CvSymKind_DEFRANGE_SUBFIELD_REGISTER:break;
+    case SYMS_CvSymKind_DEFRANGE_FRAMEPOINTER_REL_FULL_SCOPE:break;
+    case SYMS_CvSymKind_DEFRANGE_REGISTER_REL:break;
+    case SYMS_CvSymKind_LPROC32_ID:break;
+    case SYMS_CvSymKind_GPROC32_ID:break;
+    case SYMS_CvSymKind_LPROCMIPS_ID:break;
+    case SYMS_CvSymKind_GPROCMIPS_ID:break;
+    case SYMS_CvSymKind_LPROCIA64_ID:break;
+    case SYMS_CvSymKind_GPROCIA64_ID:break;
+    case SYMS_CvSymKind_BUILDINFO:break;
+    case SYMS_CvSymKind_INLINESITE:break;
+    case SYMS_CvSymKind_INLINESITE_END:break;
+    case SYMS_CvSymKind_PROC_ID_END:break;
+    case SYMS_CvSymKind_DEFRANGE_HLSL:break;
+    case SYMS_CvSymKind_GDATA_HLSL:break;
+    case SYMS_CvSymKind_LDATA_HLSL:break;
+    case SYMS_CvSymKind_FILESTATIC:break;
+    case SYMS_CvSymKind_LPROC32_DPC:break;
+    case SYMS_CvSymKind_LPROC32_DPC_ID:break;
+    case SYMS_CvSymKind_DEFRANGE_DPC_PTR_TAG:break;
+    case SYMS_CvSymKind_DPC_SYM_TAG_MAP:break;
+    case SYMS_CvSymKind_ARMSWITCHTABLE:break;
+    case SYMS_CvSymKind_CALLEES:break;
+    case SYMS_CvSymKind_CALLERS:break;
+    case SYMS_CvSymKind_POGODATA:break;
+    case SYMS_CvSymKind_INLINESITE2:break;
+    case SYMS_CvSymKind_HEAPALLOCSITE:break;
+    case SYMS_CvSymKind_MOD_TYPEREF:break;
+    case SYMS_CvSymKind_REF_MINIPDB:break;
+    case SYMS_CvSymKind_PDBMAP:break;
+    case SYMS_CvSymKind_GDATA_HLSL32:break;
+    case SYMS_CvSymKind_LDATA_HLSL32:break;
+    case SYMS_CvSymKind_GDATA_HLSL32_EX:break;
+    case SYMS_CvSymKind_LDATA_HLSL32_EX:break;
+    case SYMS_CvSymKind_FASTLINK:break;
+    case SYMS_CvSymKind_INLINEES:break;
+  }
+}
+
+
+SYMS_API SYMS_CvElement
+syms_cv_element(SYMS_String8 data, SYMS_MsfAccel *msf, SYMS_MsfRange range, SYMS_U32 off){
+  SYMS_CvSymbolHelper sym = {0};
+  syms_msf_read_struct_in_range(data, msf, range, off, &sym);
+  SYMS_CvElement result = {0};
+  if (sym.size > 0){
+    SYMS_U32 end_unclamped = off + 2 + sym.size;
+    SYMS_U32 end_clamped = SYMS_ClampTop(end_unclamped, range.size);
+    SYMS_U32 next_off = (end_unclamped + 3)&(~3);
+    SYMS_U32 start = off + 4;
+    result.range = syms_msf_sub_range(range, start, end_clamped - start);
+    result.next_off = next_off;
+    result.kind = sym.type;
+  }
+  return(result);
+}
+
+SYMS_API SYMS_U32
+syms_cv_read_numeric(SYMS_String8 data, SYMS_MsfAccel *msf, SYMS_MsfRange range, SYMS_U32 off,
+                     SYMS_CvNumeric *out){
+  SYMS_U32 result = 0;
+  SYMS_U16 leaf = 0;
+  syms_memzero_struct(out);
+  if (syms_msf_read_struct_in_range(data, msf, range, off, &leaf)){
+    if (leaf < SYMS_CvLeaf_NUMERIC){
+      out->kind = SYMS_TypeKind_UInt32;
+      syms_memmove(out->data, &leaf, 2);
+      result = sizeof(leaf);
+    }
+    else{
+      SYMS_CvNumeric num = {SYMS_TypeKind_Null};
+      SYMS_U32 size = 0;
+      switch (leaf){
+        case SYMS_CvLeaf_FLOAT16:   num.kind = SYMS_TypeKind_Float16;  size = 2;  break;
+        case SYMS_CvLeaf_FLOAT32:   num.kind = SYMS_TypeKind_Float32;  size = 4;  break;
+        case SYMS_CvLeaf_FLOAT48:   num.kind = SYMS_TypeKind_Float48;  size = 6;  break;
+        case SYMS_CvLeaf_FLOAT64:   num.kind = SYMS_TypeKind_Float64;  size = 8;  break;
+        case SYMS_CvLeaf_FLOAT80:   num.kind = SYMS_TypeKind_Float80;  size = 10; break;
+        case SYMS_CvLeaf_FLOAT128:  num.kind = SYMS_TypeKind_Float128; size = 16; break;
+        case SYMS_CvLeaf_CHAR:      num.kind = SYMS_TypeKind_Int8;     size = 1;  break;
+        case SYMS_CvLeaf_SHORT:     num.kind = SYMS_TypeKind_Int16;    size = 2;  break;
+        case SYMS_CvLeaf_USHORT:    num.kind = SYMS_TypeKind_UInt16;   size = 2;  break;
+        case SYMS_CvLeaf_LONG:      num.kind = SYMS_TypeKind_Int32;    size = 4;  break;
+        case SYMS_CvLeaf_ULONG:     num.kind = SYMS_TypeKind_UInt32;   size = 4;  break;
+        case SYMS_CvLeaf_QUADWORD:  num.kind = SYMS_TypeKind_Int64;    size = 8; break;
+        case SYMS_CvLeaf_UQUADWORD: num.kind = SYMS_TypeKind_UInt64;   size = 8; break;
+        default:break;
+      }
+      SYMS_U32 number_off = off + sizeof(leaf);
+      if (syms_msf_read_in_range(data, msf, range, number_off, size, num.data)){
+        syms_memmove(out, &num, sizeof(num));
+        result = number_off + size - off;
+      }
+    }
+  }
+  return(result);
+}
+
+SYMS_API SYMS_U32
+syms_cv_u32_from_numeric(SYMS_CvNumeric num){
+  SYMS_U32 result = 0;
+  syms_memmove(&result, num.data, 4);
+  return(result);
+}
+
+SYMS_API SYMS_ConstInfo
+syms_pdb_const_info_from_cv_numeric(SYMS_CvNumeric num){
+  SYMS_ConstInfo result = {SYMS_TypeKind_Null};
+  result.kind = num.kind;
+  syms_memmove(&result.val, num.data, sizeof(result.val));
+  return(result);
+}
+
+SYMS_API SYMS_PdbStubRef*
+syms_pdb_alloc_ref(SYMS_Arena *arena, SYMS_PdbStubRef **free_list){
+  SYMS_PdbStubRef *result = *free_list;
+  if (result != 0){
+    SYMS_StackPop(*free_list);
+  }
+  else{
+    result = syms_push_array(arena, SYMS_PdbStubRef, 1);
+  }
+  return(result);
+}
+
+
+// cv sym parse
+
+SYMS_API SYMS_PdbUnitAccel*
+syms_pdb_sym_accel_from_range(SYMS_Arena *arena, SYMS_String8 data, SYMS_MsfAccel *msf, SYMS_MsfRange range,
+                              SYMS_UnitID uid){
+  // get scratch
+  SYMS_ArenaTemp scratch = syms_get_scratch(&arena, 1);
+  
+  // root
+  SYMS_PdbStub root = {0};
+  SYMS_U64 top_count = 0;
+  
+  // all list
+  SYMS_PdbStub *first = 0;
+  SYMS_PdbStub *last = 0;
+  SYMS_U64 all_count = 0;
+  
+  // proc & var
+  SYMS_PdbStubRef *proc_first = 0;
+  SYMS_PdbStubRef *proc_last = 0;
+  SYMS_U64 proc_count = 0;
+  
+  SYMS_PdbStubRef *var_first = 0;
+  SYMS_PdbStubRef *var_last = 0;
+  SYMS_U64 var_count = 0;
+  
+  // parse loop
+  SYMS_PdbStub *defrange_collector_stub = 0;
+  SYMS_PdbStubRef *stack = 0;
+  SYMS_PdbStubRef *stack_free = 0;
+  SYMS_U32 cursor = 0;
+  for (;;){
+    // read element
+    SYMS_CvElement element = syms_cv_element(data, msf, range, cursor);
+    
+    // exit condition
+    if (element.next_off == 0){
+      break;
+    }
+    
+    // init stub
+    SYMS_U32 symbol_off = element.range.off - 4;
+    SYMS_PdbStub *stub = syms_push_array(arena, SYMS_PdbStub, 1);
+    syms_memzero_struct(stub);
+    SYMS_QueuePush_N(first, last, stub, bucket_next);
+    all_count += 1;
+    stub->off = symbol_off;
+    
+    // default parent
+    SYMS_PdbStub *parent_for_this_stub = &root;
+    if (stack != 0){
+      parent_for_this_stub = stack->stub;
+    }
+    
+    // defrange collecting check
+    SYMS_B32 preserve_defrange_collection = syms_false;
+    
+    // parse
+    switch (element.kind){
+      default:break;
+      
+      case SYMS_CvSymKind_LDATA32:
+      case SYMS_CvSymKind_GDATA32:
+      {
+        SYMS_U32 name_off = sizeof(SYMS_CvData32);
+        SYMS_String8 name = syms_msf_read_zstring_in_range(arena, data, msf, element.range, name_off);
+        
+        stub->name = name;
+        
+        // push onto var list
+        SYMS_PdbStubRef *ref = syms_pdb_alloc_ref(scratch.arena, &stack_free);
+        ref->stub = stub;
+        SYMS_QueuePush(var_first, var_last, ref);
+        var_count += 1;
+      }break;
+      
+      // general local variable
+      case SYMS_CvSymKind_LOCAL:
+      {
+        SYMS_U32 name_off = sizeof(SYMS_CvLocal);
+        SYMS_String8 name = syms_msf_read_zstring_in_range(arena, data, msf, element.range, name_off);
+        
+        stub->name = name;
+        
+        preserve_defrange_collection = syms_true;
+        defrange_collector_stub = stub;
+      }break;
+      
+      case SYMS_CvSymKind_DEFRANGE_2005:
+      case SYMS_CvSymKind_DEFRANGE2_2005:
+      case SYMS_CvSymKind_DEFRANGE:
+      case SYMS_CvSymKind_DEFRANGE_SUBFIELD:
+      case SYMS_CvSymKind_DEFRANGE_REGISTER:
+      case SYMS_CvSymKind_DEFRANGE_FRAMEPOINTER_REL:
+      case SYMS_CvSymKind_DEFRANGE_SUBFIELD_REGISTER:
+      case SYMS_CvSymKind_DEFRANGE_FRAMEPOINTER_REL_FULL_SCOPE:
+      case SYMS_CvSymKind_DEFRANGE_REGISTER_REL:
+      {
+        preserve_defrange_collection = syms_true;
+        parent_for_this_stub = defrange_collector_stub;
+      }break;
+      
+      case SYMS_CvSymKind_REGREL32:
+      {
+        SYMS_U32 name_off = sizeof(SYMS_CvRegrel32);
+        SYMS_String8 name = syms_msf_read_zstring_in_range(arena, data, msf, element.range, name_off);
+        
+        stub->name = name;
+      }break;
+      
+      case SYMS_CvSymKind_CONSTANT:
+      {
+        SYMS_CvNumeric numeric_val = {SYMS_TypeKind_Null};
+        SYMS_String8 name = {0};
+        SYMS_U32 numeric_off = sizeof(SYMS_CvConstant);
+        SYMS_U32 numeric_size = syms_cv_read_numeric(data, msf, element.range, numeric_off, &numeric_val);
+        SYMS_U32 name_off = numeric_off + numeric_size;
+        name = syms_msf_read_zstring_in_range(arena, data, msf, element.range, name_off);
+        
+        stub->num = syms_cv_u32_from_numeric(numeric_val);
+        stub->name = name;
+      }break;
+      
+      case SYMS_CvSymKind_PROCREF:
+      case SYMS_CvSymKind_LPROCREF:
+      case SYMS_CvSymKind_DATAREF:
+      {
+        SYMS_U32 name_off = sizeof(SYMS_CvRef2);
+        SYMS_String8 name = syms_msf_read_zstring_in_range(arena, data, msf, element.range, name_off);
+        
+        stub->name = name;
+      }break;
+      
+      case SYMS_CvSymKind_LPROC32:
+      case SYMS_CvSymKind_GPROC32:
+      {
+        SYMS_U32 name_off = sizeof(SYMS_CvProc32);
+        SYMS_String8 name = syms_msf_read_zstring_in_range(arena, data, msf, element.range, name_off);
+        
+        stub->name = name;
+        
+        // push stub stack
+        SYMS_PdbStubRef *stack_node = syms_pdb_alloc_ref(scratch.arena, &stack_free);
+        SYMS_StackPush(stack, stack_node);
+        stack_node->stub = stub;
+        
+        // push onto proc list
+        SYMS_PdbStubRef *ref = syms_pdb_alloc_ref(scratch.arena, &stack_free);
+        ref->stub = stub;
+        SYMS_QueuePush(proc_first, proc_last, ref);
+        proc_count += 1;
+      }break;
+      
+      case SYMS_CvSymKind_BLOCK32:
+      {
+        SYMS_U32 name_off = sizeof(SYMS_CvBlock32);
+        SYMS_String8 name = syms_msf_read_zstring_in_range(arena, data, msf, element.range, name_off);
+        
+        stub->name = name;
+        
+        // push stub stack
+        SYMS_PdbStubRef *stack_node = syms_pdb_alloc_ref(scratch.arena, &stack_free);
+        SYMS_StackPush(stack, stack_node);
+        stack_node->stub = stub;
+      }break;
+      
+      case SYMS_CvSymKind_END:
+      {
+        if (stack != 0){
+          SYMS_PdbStubRef *bucket = stack;
+          SYMS_StackPop(stack);
+          SYMS_StackPush(stack_free, bucket);
+        }
+        parent_for_this_stub = 0;
+      }break;
+      
+      case SYMS_CvSymKind_PUB32:
+      {
+        SYMS_U32 name_off = sizeof(SYMS_CvPubsym32);
+        SYMS_String8 name = syms_msf_read_zstring_in_range(scratch.arena, data, msf, element.range, name_off);
+        
+        stub->name = name;
+      }break;
+    }
+    
+    // clear defrange collector
+    if (!preserve_defrange_collection){
+      defrange_collector_stub = 0;
+    }
+    
+    // insert into tree
+    if (parent_for_this_stub != 0){
+      SYMS_QueuePush_N(parent_for_this_stub->first, parent_for_this_stub->last, stub, sibling_next);
+    }
+    if (parent_for_this_stub == &root){
+      top_count += 1;
+    }
+    
+    // increment
+    cursor = element.next_off;
+  }
+  
+  // build top stubs pointer table
+  SYMS_PdbStub **top_stubs = syms_push_array(arena, SYMS_PdbStub*, top_count);
+  {
+    SYMS_PdbStub **ptr = top_stubs;
+    for (SYMS_PdbStub *node = root.first;
+         node != 0;
+         node = node->sibling_next, ptr += 1){
+      *ptr = node;
+    }
+  }
+  
+  // build proc stubs pointer table
+  SYMS_PdbStub **proc_stubs = syms_push_array(arena, SYMS_PdbStub*, proc_count);
+  {
+    SYMS_PdbStub **ptr = proc_stubs;
+    for (SYMS_PdbStubRef *ref = proc_first;
+         ref != 0;
+         ref = ref->next, ptr += 1){
+      *ptr = ref->stub;
+    }
+  }
+  
+  // build var stubs pointer table
+  SYMS_PdbStub **var_stubs = syms_push_array(arena, SYMS_PdbStub*, var_count);
+  {
+    SYMS_PdbStub **ptr = var_stubs;
+    for (SYMS_PdbStubRef *ref = var_first;
+         ref != 0;
+         ref = ref->next, ptr += 1){
+      *ptr = ref->stub;
+    }
+  }
+  
+  // build bucket table
+  SYMS_U64 bucket_count = (all_count/2)*2 + 3;
+  SYMS_PdbStub **buckets = syms_push_array(arena, SYMS_PdbStub*, bucket_count);
+  if (bucket_count > 0){
+    syms_memset(buckets, 0, sizeof(*buckets)*bucket_count);
+    for (SYMS_PdbStub *bucket = first, *next = 0;
+         bucket != 0;
+         bucket = next){
+      next = bucket->bucket_next;
+      SYMS_U64 hash = syms_hash_u64(bucket->off);
+      SYMS_U32 bucket_index = hash%bucket_count;
+      SYMS_StackPush_N(buckets[bucket_index], bucket, bucket_next);
+    }
+  }
+  
+  // release scratch
+  syms_release_scratch(scratch);
+  
+  // fill result
+  SYMS_PdbUnitAccel *result = syms_push_array_zero(arena, SYMS_PdbUnitAccel, 1);
+  result->format = SYMS_FileFormat_PDB;
+  result->sn = range.sn;
+  result->top_stubs = top_stubs;
+  result->top_count = top_count;
+  result->buckets = buckets;
+  result->bucket_count = bucket_count;
+  result->all_count = all_count;
+  result->uid = uid;
+  result->proc_stubs = proc_stubs;
+  result->proc_count = proc_count;
+  result->var_stubs = var_stubs;
+  result->var_count = var_count;
+  
+  return(result);
+}
+
+SYMS_API SYMS_PdbUnitAccel*
+syms_pdb_pub_sym_accel_from_dbg(SYMS_Arena *arena, SYMS_String8 data, SYMS_PdbDbgAccel *dbg){
+  SYMS_MsfAccel *msf = dbg->msf;
+  SYMS_MsfRange range = syms_msf_range_from_sn(msf, dbg->dbi.sym_sn);
+  SYMS_PdbUnitAccel *result = syms_pdb_sym_accel_from_range(arena, data, msf, range, SYMS_PdbPseudoUnit_SYM);
+  return(result);
+}
+
+SYMS_API SYMS_SymbolKind
+syms_pdb_sym_symbol_kind_from_id(SYMS_String8 data, SYMS_PdbDbgAccel *dbg,
+                                 SYMS_PdbUnitAccel *unit, SYMS_SymbolID id){
+  //- setup accel
+  SYMS_MsfAccel *msf = dbg->msf;
+  
+  //- read id
+  SYMS_MsfRange range = syms_msf_range_from_sn(msf, unit->sn);
+  SYMS_PdbStub *stub = syms_pdb_stub_from_unit_off(unit, SYMS_ID_u32_1(id));
+  
+  //- parse symbol
+  SYMS_SymbolKind result = SYMS_SymbolKind_Null;
+  if (stub != 0){
+    //- get kind
+    SYMS_CvElement element = syms_cv_element(data, msf, range, stub->off);
+    switch (element.kind){
+      default:break;
+      case SYMS_CvSymKind_LDATA32:
+      case SYMS_CvSymKind_GDATA32:
+      {
+        result = SYMS_SymbolKind_ImageRelativeVariable;
+      }break;
+      case SYMS_CvSymKind_LOCAL:
+      {
+        result = SYMS_SymbolKind_LocalVariable;
+      }break;
+      case SYMS_CvSymKind_REGREL32:
+      {
+        result = SYMS_SymbolKind_LocalVariable;
+      }break;
+      case SYMS_CvSymKind_CONSTANT:
+      {
+        result = SYMS_SymbolKind_Const;
+      }break;
+      case SYMS_CvSymKind_LPROC32:
+      case SYMS_CvSymKind_GPROC32:
+      {
+        result = SYMS_SymbolKind_Procedure;
+      }break;
+      case SYMS_CvSymKind_BLOCK32:
+      {
+        result = SYMS_SymbolKind_Scope;
+      }break;
+    }
+  }
+  
+  return(result);
+}
+
+SYMS_API SYMS_String8
+syms_pdb_sym_symbol_name_from_id(SYMS_Arena *arena, SYMS_String8 data, SYMS_PdbDbgAccel *dbg,
+                                 SYMS_PdbUnitAccel *unit, SYMS_SymbolID id){
+  //- setup accel
+  SYMS_MsfAccel *msf = dbg->msf;
+  
+  //- read id
+  SYMS_MsfRange range = syms_msf_range_from_sn(msf, unit->sn);
+  SYMS_PdbStub *stub = syms_pdb_stub_from_unit_off(unit, SYMS_ID_u32_1(id));
+  
+  //- parse symbol
+  SYMS_String8 result = {0};
+  if (stub != 0){
+    //- copy name if needed
+    result = syms_push_string_copy(arena, stub->name);
+  }
+  
+  return(result);
+}
+
+// cv leaf parse
+
+SYMS_API SYMS_PdbUnitAccel*
+syms_pdb_leaf_accel_from_dbg(SYMS_Arena *arena, SYMS_String8 data, SYMS_PdbDbgAccel *dbg){
+  // setup accel
+  SYMS_MsfAccel *msf = dbg->msf;
+  SYMS_PdbTpiAccel *tpi = &dbg->tpi;
+  SYMS_MsfRange range = syms_pdb_tpi_range(msf, tpi);
+  
+  // allocate stub array
+  SYMS_CvTypeIndex first_ti = tpi->first_ti;
+  SYMS_CvTypeIndex opl_ti = tpi->opl_ti;
+  SYMS_U64 ti_count = opl_ti - first_ti;
+  SYMS_PdbStub *ti_stubs = syms_push_array(arena, SYMS_PdbStub, ti_count);
+  
+  // initialize off stub count
+  SYMS_U64 off_stub_count = 0;
+  
+  // parse loop
+  SYMS_PdbStub *stub = ti_stubs;
+  SYMS_PdbStub *opl = ti_stubs + ti_count;
+  SYMS_U32 cursor = 0;
+  SYMS_CvTypeIndex ti = first_ti;
+  for (;stub < opl;){
+    // read element
+    SYMS_CvElement element = syms_cv_element(data, msf, range, cursor);
+    
+    // exit condition
+    if (element.next_off == 0){
+      break;
+    }
+    
+    // setup stub
+    syms_memzero_struct(stub);
+    stub->off = element.range.off - 4;
+    
+    // parse
+    switch (element.kind){
+      default:break;
+      
+      case SYMS_CvLeaf_ARRAY:
+      {
+        SYMS_U32 num_off = sizeof(SYMS_CvLeafArray);
+        SYMS_CvNumeric num = {SYMS_TypeKind_Null};
+        SYMS_U32 num_size = syms_cv_read_numeric(data, msf, element.range, num_off, &num);
+        (void)num_size;
+        stub->num = syms_cv_u32_from_numeric(num);
+      }break;
+      
+      // udt
+      case SYMS_CvLeaf_CLASS:
+      case SYMS_CvLeaf_STRUCTURE:
+      case SYMS_CvLeaf_INTERFACE:
+      {
+        SYMS_U32 num_off = sizeof(SYMS_CvLeafStruct);
+        SYMS_CvNumeric num = {SYMS_TypeKind_Null};
+        SYMS_U32 num_size = syms_cv_read_numeric(data, msf, element.range, num_off, &num);
+        SYMS_U32 name_off = num_off + num_size;
+        SYMS_String8 name = syms_msf_read_zstring_in_range(arena, data, msf, element.range, name_off);
+        
+        stub->num = syms_cv_u32_from_numeric(num);
+        stub->name = name;
+      }break;
+      
+      case SYMS_CvLeaf_UNION:
+      {
+        SYMS_U32 num_off = sizeof(SYMS_CvLeafUnion);
+        SYMS_CvNumeric num = {SYMS_TypeKind_Null};
+        SYMS_U32 num_size = syms_cv_read_numeric(data, msf, element.range, num_off, &num);
+        SYMS_U32 name_off = num_off + num_size;
+        SYMS_String8 name = syms_msf_read_zstring_in_range(arena, data, msf, element.range, name_off);
+        
+        stub->num = syms_cv_u32_from_numeric(num);
+        stub->name = name;
+      }break;
+      
+      case SYMS_CvLeaf_ENUM:
+      {
+        SYMS_U32 name_off = sizeof(SYMS_CvLeafEnum);
+        SYMS_String8 name = syms_msf_read_zstring_in_range(arena, data, msf, element.range, name_off);
+        stub->name = name;
+      }break;
+      
+      case SYMS_CvLeaf_CLASSPTR:
+      case SYMS_CvLeaf_CLASSPTR2:
+      {
+        SYMS_U32 num_off = sizeof(SYMS_CvLeafClassPtr);
+        SYMS_CvNumeric num = {SYMS_TypeKind_Null};
+        SYMS_U32 num_size = syms_cv_read_numeric(data, msf, element.range, num_off, &num);
+        SYMS_U32 name_off = num_off + num_size;
+        SYMS_String8 name = syms_msf_read_zstring_in_range(arena, data, msf, element.range, name_off);
+        
+        stub->num = syms_cv_u32_from_numeric(num);
+        stub->name = name;
+      }break;
+      
+      case SYMS_CvLeaf_ALIAS:
+      {
+        SYMS_U32 name_off = sizeof(SYMS_CvLeafAlias);
+        SYMS_String8 name = syms_msf_read_zstring_in_range(arena, data, msf, element.range, name_off);
+        
+        stub->name = name;
+      }break;
+      
+      // field list
+      case SYMS_CvLeaf_FIELDLIST:
+      {
+        // parse loop
+        SYMS_PdbStub *parent = stub;
+        SYMS_U32 fl_cursor = 0;
+        for (;fl_cursor < element.range.size;){
+          // read kind
+          SYMS_CvLeaf lf_kind = 0;
+          syms_msf_read_struct_in_range(data, msf, element.range, fl_cursor, &lf_kind);
+          
+          // insert new stub under parent
+          SYMS_PdbStub *fl_stub = syms_push_array(arena, SYMS_PdbStub, 1);
+          syms_memzero_struct(fl_stub);
+          SYMS_QueuePush_N(parent->first, parent->last, fl_stub, sibling_next);
+          fl_stub->off = element.range.off + fl_cursor;
+          
+          // read internal leaf
+          SYMS_U32 lf_data_off = fl_cursor + sizeof(lf_kind);
+          SYMS_U32 lf_end_off = lf_data_off + 2;
+          
+          switch (lf_kind){
+            default:break;
+            
+            case SYMS_CvLeaf_MEMBER:
+            {
+              SYMS_U32 num_off = lf_data_off + sizeof(SYMS_CvLeafMember);
+              SYMS_CvNumeric num = {SYMS_TypeKind_Null};
+              SYMS_U32 num_size = syms_cv_read_numeric(data, msf, element.range, num_off, &num);
+              SYMS_U32 name_off = num_off + num_size;
+              SYMS_String8 name = syms_msf_read_zstring_in_range(arena, data, msf, element.range, name_off);
+              
+              fl_stub->num = syms_cv_u32_from_numeric(num);
+              fl_stub->name = name;
+              lf_end_off = name_off + name.size + 1;
+            }break;
+            
+            case SYMS_CvLeaf_STMEMBER:
+            {
+              SYMS_U32 name_off = lf_data_off + sizeof(SYMS_CvLeafStMember);
+              SYMS_String8 name = syms_msf_read_zstring_in_range(arena, data, msf, element.range, name_off);
+              
+              fl_stub->name = name;
+              lf_end_off = name_off + name.size + 1;
+            }break;
+            
+            case SYMS_CvLeaf_METHOD:
+            {
+              SYMS_U32 name_off = lf_data_off + sizeof(SYMS_CvLeafMethod);
+              SYMS_String8 name = syms_msf_read_zstring_in_range(arena, data, msf, element.range, name_off);
+              
+              fl_stub->name = name;
+              lf_end_off = name_off + name.size + 1;
+            }break;
+            
+            case SYMS_CvLeaf_ONEMETHOD:
+            {
+              SYMS_CvLeafOneMethod onemethod = {0};
+              syms_msf_read_struct_in_range(data, msf, element.range, lf_data_off, &onemethod);
+              SYMS_CvMethodProp prop = SYMS_CvFieldAttribs_Extract_MPROP(onemethod.attribs);
+              
+              SYMS_U32 name_off = lf_data_off + sizeof(onemethod);
+              SYMS_U32 virtoff = 0;
+              if (prop == SYMS_CvMethodProp_PUREINTRO || prop == SYMS_CvMethodProp_INTRO){
+                syms_msf_read_struct_in_range(data, msf, range, name_off, &virtoff);
+                name_off += sizeof(virtoff);
+              }
+              SYMS_String8 name = syms_msf_read_zstring_in_range(arena, data, msf, element.range, name_off);
+              
+              fl_stub->name = name;
+              fl_stub->num = virtoff;
+              lf_end_off = name_off + name.size + 1;
+            }break;
+            
+            case SYMS_CvLeaf_ENUMERATE:
+            {
+              SYMS_U32 num_off = lf_data_off + sizeof(SYMS_CvLeafEnumerate);
+              SYMS_CvNumeric num = {SYMS_TypeKind_Null};
+              SYMS_U32 num_size = syms_cv_read_numeric(data, msf, element.range, num_off, &num);
+              SYMS_U32 name_off = num_off + num_size;
+              SYMS_String8 name = syms_msf_read_zstring_in_range(arena, data, msf, element.range, name_off);
+              
+              fl_stub->num = syms_cv_u32_from_numeric(num);
+              fl_stub->name = name;
+              lf_end_off = name_off + name.size + 1;
+            }break;
+            
+            case SYMS_CvLeaf_NESTTYPE:
+            case SYMS_CvLeaf_NESTTYPEEX:
+            {
+              SYMS_U32 name_off = lf_data_off + sizeof(SYMS_CvLeafNestTypeEx);
+              SYMS_String8 name = syms_msf_read_zstring_in_range(arena, data, msf, element.range, name_off);
+              
+              fl_stub->name = name;
+              lf_end_off = name_off + name.size + 1;
+            }break;
+            
+            case SYMS_CvLeaf_BCLASS:
+            {
+              SYMS_U32 num_off = lf_data_off + sizeof(SYMS_CvLeafBClass);
+              SYMS_CvNumeric num = {SYMS_TypeKind_Null};
+              SYMS_U32 num_size = syms_cv_read_numeric(data, msf, element.range, num_off, &num);
+              
+              fl_stub->num = syms_cv_u32_from_numeric(num);
+              lf_end_off = num_off + num_size;
+            }break;
+            
+            case SYMS_CvLeaf_VBCLASS:
+            case SYMS_CvLeaf_IVBCLASS:
+            {
+              SYMS_U32 num_off = lf_data_off + sizeof(SYMS_CvLeafVBClass);
+              SYMS_CvNumeric num = {SYMS_TypeKind_Null};
+              SYMS_U32 num_size = syms_cv_read_numeric(data, msf, element.range, num_off, &num);
+              
+              SYMS_U32 num2_off = num_off + num_size;
+              SYMS_CvNumeric num2 = {SYMS_TypeKind_Null};
+              SYMS_U32 num2_size = syms_cv_read_numeric(data, msf, element.range, num2_off, &num2);
+              
+              fl_stub->num = syms_cv_u32_from_numeric(num);
+              fl_stub->num2 = syms_cv_u32_from_numeric(num2);
+              lf_end_off = num2_off + num2_size;
+            }break;
+            
+            case SYMS_CvLeaf_VFUNCTAB:
+            {
+              lf_end_off = lf_data_off + sizeof(SYMS_CvLeafVFuncTab);
+            }break;
+            
+            case SYMS_CvLeaf_VFUNCOFF:
+            {
+              lf_end_off = lf_data_off + sizeof(SYMS_CvLeafVFuncOff);
+            }break;
+          }
+          
+          // increment
+          SYMS_U32 next_lf_off = (lf_end_off + 3)&~3;
+          fl_cursor = next_lf_off;
+          off_stub_count += 1;
+        }
+        
+      }break;
+      
+      // method list
+      case SYMS_CvLeaf_METHODLIST:
+      {
+        // parse loop
+        SYMS_PdbStub *parent = stub;
+        SYMS_U32 ml_cursor = 0;
+        
+        for (;ml_cursor < element.range.size;){
+          // read method
+          SYMS_CvMethod methodrec = {0};
+          syms_msf_read_struct_in_range(data, msf, element.range, ml_cursor, &methodrec);
+          SYMS_CvMethodProp mprop = SYMS_CvFieldAttribs_Extract_MPROP(methodrec.attribs);
+          
+          // get virtual offset
+          SYMS_U32 virtual_offset = 0;
+          SYMS_U32 next_off = ml_cursor + sizeof(methodrec);
+          switch (mprop){
+            default:break;
+            case SYMS_CvMethodProp_INTRO:
+            case SYMS_CvMethodProp_PUREINTRO:
+            {
+              syms_msf_read_struct_in_range(data, msf, element.range, next_off, &virtual_offset);
+              next_off += sizeof(SYMS_U32);
+            }break;
+          }
+          
+          // insert new stub under parent
+          SYMS_PdbStub *ml_stub = syms_push_array(arena, SYMS_PdbStub, 1);
+          syms_memzero_struct(ml_stub);
+          SYMS_QueuePush_N(parent->first, parent->last, ml_stub, sibling_next);
+          ml_stub->off = element.range.off + ml_cursor;
+          ml_stub->num = virtual_offset;
+          
+          // increment
+          ml_cursor = next_off;
+        }
+      }break;
+    }
+    
+    // increment
+    cursor = element.next_off;
+    ti += 1;
+    stub += 1;
+  }
+  
+  // build bucket table
+  SYMS_U64 bucket_count = off_stub_count*5/4;
+  SYMS_PdbStub **buckets = syms_push_array(arena, SYMS_PdbStub*, bucket_count);
+  if (bucket_count > 0){
+    syms_memset(buckets, 0, sizeof(*buckets)*bucket_count);
+    SYMS_PdbStub *stub = ti_stubs;
+    for (SYMS_U64 i = 0; i < ti_count; i += 1, stub += 1){
+      for (SYMS_PdbStub *internal_stub = stub->first;
+           internal_stub != 0;
+           internal_stub = internal_stub->sibling_next){
+        SYMS_U64 hash = syms_hash_u64(internal_stub->off);
+        SYMS_U32 bucket_index = hash%bucket_count;
+        SYMS_StackPush_N(buckets[bucket_index], internal_stub, bucket_next);
+      }
+    }
+  }
+  
+  // filter top stubs
+  SYMS_MsfRange whole_range = syms_msf_range_from_sn(msf, range.sn);
+  SYMS_PdbStub **udt_stubs = syms_push_array(arena, SYMS_PdbStub*, ti_count);
+  SYMS_U64 udt_count = 0;
+  {
+    SYMS_PdbStub **top_stub_ptr = udt_stubs;
+    SYMS_PdbStub *stub = ti_stubs;
+    SYMS_PdbStub *stub_opl = ti_stubs + ti_count;
+    for (; stub < stub_opl; stub += 1){
+      SYMS_U16 type = {0};
+      syms_msf_read_struct_in_range(data, msf, whole_range, stub->off + 2, &type);
+      switch (type){
+        // type kinds
+        case SYMS_CvLeaf_CLASS:
+        case SYMS_CvLeaf_STRUCTURE:
+        case SYMS_CvLeaf_INTERFACE:
+        case SYMS_CvLeaf_UNION:
+        case SYMS_CvLeaf_ENUM:
+        case SYMS_CvLeaf_CLASSPTR:
+        case SYMS_CvLeaf_CLASSPTR2:
+        case SYMS_CvLeaf_ALIAS:
+        {
+          *top_stub_ptr = stub;
+          top_stub_ptr += 1;
+        }break;
+      }
+    }
+    
+    udt_count = (SYMS_U64)(top_stub_ptr - udt_stubs);
+    syms_arena_put_back(arena, sizeof(SYMS_PdbStub*)*(ti_count - udt_count));
+  }
+  
+  // fill result
+  SYMS_PdbUnitAccel *result = syms_push_array(arena, SYMS_PdbUnitAccel, 1);
+  result->format = SYMS_FileFormat_PDB;
+  result->leaf_set = syms_true;
+  result->sn = range.sn;
+  result->top_stubs = udt_stubs;
+  result->top_count = udt_count;
+  result->top_min_index = first_ti;
+  result->buckets = buckets;
+  result->bucket_count = bucket_count;
+  result->all_count = ti_count + off_stub_count;
+  result->ti_stubs = ti_stubs;
+  result->ti_count = ti_count;
+  result->uid = SYMS_PdbPseudoUnit_TPI;
+  result->udt_stubs = udt_stubs;
+  result->udt_count = udt_count;
+  
+  return(result);
+}
+
+SYMS_API SYMS_PdbLeafResolve
+syms_pdb_leaf_resolve_from_id(SYMS_String8 data, SYMS_PdbDbgAccel *dbg, SYMS_PdbUnitAccel *unit, SYMS_SymbolID id){
+  // setup accel
+  SYMS_MsfAccel *msf = dbg->msf;
+  
+  // read id
+  SYMS_MsfRange range = syms_msf_range_from_sn(msf, unit->sn);
+  
+  SYMS_PdbLeafResolve result = {0};
+  switch (SYMS_ID_u16_0(id)){
+    case SYMS_PdbSymbolIDKind_Index:
+    {
+      result.stub = syms_pdb_stub_from_unit_index(unit, SYMS_ID_u32_1(id));
+      if (result.stub != 0){
+        SYMS_CvElement element = syms_cv_element(data, msf, range, result.stub->off);
+        result.leaf_kind = element.kind;
+        result.leaf_range = element.range;
+      }
+      result.is_leaf_id = syms_true;
+    }break;
+    
+    case SYMS_PdbSymbolIDKind_Off:
+    {
+      result.stub = syms_pdb_stub_from_unit_off(unit, SYMS_ID_u32_1(id));
+      if (result.stub != 0){
+        SYMS_U32 lf_off = result.stub->off;
+        syms_msf_read_struct_in_range(data, msf, range, lf_off, &result.leaf_kind);
+        SYMS_U32 lf_data_off = lf_off + 2;
+        result.leaf_range = syms_msf_sub_range(range, lf_data_off, range.size - lf_data_off);
+      }
+    }break;
+  }
+  
+  return(result);
+}
+
+SYMS_API SYMS_SymbolKind
+syms_pdb_leaf_symbol_kind_from_id(SYMS_String8 data, SYMS_PdbDbgAccel *dbg,
+                                  SYMS_PdbUnitAccel *unit, SYMS_SymbolID id){
+  //- read id
+  SYMS_PdbLeafResolve resolve = syms_pdb_leaf_resolve_from_id(data, dbg, unit, id);
+  
+  //- zero clear result
+  SYMS_SymbolKind result = SYMS_SymbolKind_Null;
+  
+  //- basic type info
+  if (resolve.stub == 0 && resolve.is_leaf_id){
+    SYMS_CvTypeIndex ti = SYMS_ID_u32_1(id);
+    if (ti == syms_cv_type_id_variadic){
+      result = SYMS_SymbolKind_Type;
+    }
+    else if (ti < 0x1000){
+      SYMS_CvBasicPointerKind basic_ptr_kind = SYMS_CvBasicPointerKindFromTypeId(ti);
+      switch (basic_ptr_kind){
+        case SYMS_CvBasicPointerKind_VALUE:
+        case SYMS_CvBasicPointerKind_16BIT:
+        case SYMS_CvBasicPointerKind_FAR_16BIT:
+        case SYMS_CvBasicPointerKind_HUGE_16BIT:
+        case SYMS_CvBasicPointerKind_32BIT:
+        case SYMS_CvBasicPointerKind_16_32BIT:
+        case SYMS_CvBasicPointerKind_64BIT:
+        {
+          SYMS_U32 itype_kind = SYMS_CvBasicTypeFromTypeId(ti);
+          switch (itype_kind){
+            case SYMS_CvBasicType_VOID:
+            case SYMS_CvBasicType_HRESULT:
+            case SYMS_CvBasicType_RCHAR:
+            case SYMS_CvBasicType_CHAR:
+            case SYMS_CvBasicType_UCHAR:
+            case SYMS_CvBasicType_WCHAR:
+            case SYMS_CvBasicType_SHORT:
+            case SYMS_CvBasicType_USHORT:
+            case SYMS_CvBasicType_LONG:
+            case SYMS_CvBasicType_ULONG:
+            case SYMS_CvBasicType_QUAD:
+            case SYMS_CvBasicType_UQUAD:
+            case SYMS_CvBasicType_OCT:
+            case SYMS_CvBasicType_UOCT:
+            case SYMS_CvBasicType_CHAR8:
+            case SYMS_CvBasicType_CHAR16:
+            case SYMS_CvBasicType_CHAR32:
+            case SYMS_CvBasicType_BOOL8:
+            case SYMS_CvBasicType_BOOL16:
+            case SYMS_CvBasicType_BOOL32:
+            case SYMS_CvBasicType_BOOL64:
+            case SYMS_CvBasicType_INT8:
+            case SYMS_CvBasicType_INT16:
+            case SYMS_CvBasicType_INT32:
+            case SYMS_CvBasicType_INT64:
+            case SYMS_CvBasicType_INT128:
+            case SYMS_CvBasicType_UINT8:
+            case SYMS_CvBasicType_UINT16:
+            case SYMS_CvBasicType_UINT32:
+            case SYMS_CvBasicType_UINT64:
+            case SYMS_CvBasicType_UINT128:
+            case SYMS_CvBasicType_FLOAT16:
+            case SYMS_CvBasicType_FLOAT32:
+            case SYMS_CvBasicType_FLOAT64:
+            case SYMS_CvBasicType_FLOAT32PP:
+            case SYMS_CvBasicType_FLOAT80:
+            case SYMS_CvBasicType_FLOAT128:
+            case SYMS_CvBasicType_COMPLEX32:
+            case SYMS_CvBasicType_COMPLEX64:
+            case SYMS_CvBasicType_COMPLEX80:
+            case SYMS_CvBasicType_COMPLEX128:
+            {
+              result = SYMS_SymbolKind_Type;
+            }break;
+          }
+        }break;
+      }
+    }
+  }
+  
+  //- recorded type info
+  if (resolve.stub != 0){
+    switch (resolve.leaf_kind){
+      // type kinds
+      case SYMS_CvLeaf_MODIFIER:
+      case SYMS_CvLeaf_POINTER:
+      case SYMS_CvLeaf_PROCEDURE:
+      case SYMS_CvLeaf_MFUNCTION:
+      case SYMS_CvLeaf_ARRAY:
+      case SYMS_CvLeaf_BITFIELD:
+      case SYMS_CvLeaf_CLASS:
+      case SYMS_CvLeaf_STRUCTURE:
+      case SYMS_CvLeaf_INTERFACE:
+      case SYMS_CvLeaf_UNION:
+      case SYMS_CvLeaf_ENUM:
+      case SYMS_CvLeaf_CLASSPTR:
+      case SYMS_CvLeaf_CLASSPTR2:
+      case SYMS_CvLeaf_ALIAS:
+      {
+        result = SYMS_SymbolKind_Type;
+      }break;
+    }
+  }
+  
+  return(result);
+}
+
+SYMS_API SYMS_String8
+syms_pdb_leaf_symbol_name_from_id(SYMS_Arena *arena, SYMS_String8 data, SYMS_PdbDbgAccel *dbg,
+                                  SYMS_PdbUnitAccel *unit, SYMS_SymbolID id){
+  //- read id
+  SYMS_PdbLeafResolve resolve = syms_pdb_leaf_resolve_from_id(data, dbg, unit, id);
+  
+  //- zero clear result
+  SYMS_String8 result = {0};
+  
+  //- basic type info
+  if (resolve.stub == 0 && resolve.is_leaf_id){
+    SYMS_CvTypeIndex ti = SYMS_ID_u32_1(id);
+    if (ti < 0x1000){
+      SYMS_CvBasicPointerKind basic_ptr_kind = SYMS_CvBasicPointerKindFromTypeId(ti);
+      if (basic_ptr_kind == SYMS_CvBasicPointerKind_VALUE){
+        SYMS_U32 itype_kind = SYMS_CvBasicTypeFromTypeId(ti);
+        result = syms_string_from_enum_value(SYMS_CvBasicType, itype_kind);
+      }
+    }
+  }
+  
+  //- recorded type info
+  if (resolve.stub != 0){
+    result = syms_push_string_copy(arena, resolve.stub->name);
+  }
+  
+  return(result);
+}
+
+SYMS_API SYMS_TypeInfo
+syms_pdb_leaf_type_info_from_id(SYMS_String8 data, SYMS_PdbDbgAccel *dbg, SYMS_PdbUnitAccel *unit,
+                                SYMS_SymbolID id){
+  // setup accel
+  SYMS_MsfAccel *msf = dbg->msf;
+  
+  // read id
+  SYMS_PdbLeafResolve resolve = syms_pdb_leaf_resolve_from_id(data, dbg, unit, id);
+  
+  // zero clear result
+  SYMS_TypeInfo result;
+  syms_memzero_struct(&result);
+  
+  // basic type info
+  if (resolve.stub == 0 && resolve.is_leaf_id){
+    SYMS_CvTypeIndex ti = SYMS_ID_u32_1(id);
+    if (ti == syms_cv_type_id_variadic){
+      result.kind = SYMS_TypeKind_Variadic;
+    }
+    else if (ti < 0x1000){
+      SYMS_CvBasicPointerKind basic_ptr_kind = SYMS_CvBasicPointerKindFromTypeId(ti);
+      SYMS_U32 itype_kind = SYMS_CvBasicTypeFromTypeId(ti);
+      
+      result.reported_size_interp = SYMS_SizeInterpretation_ByteCount;
+      
+      switch (basic_ptr_kind){
+        default:
+        {
+          result.reported_size_interp = SYMS_SizeInterpretation_Null;
+        }break;
+        
+        case SYMS_CvBasicPointerKind_VALUE:
+        {
+          switch (itype_kind){
+            default:
+            {
+              result.reported_size_interp = SYMS_SizeInterpretation_Null;
+            }break;
+            
+            case SYMS_CvBasicType_VOID:
+            {
+              result.kind = SYMS_TypeKind_Void;
+              result.reported_size_interp = SYMS_SizeInterpretation_Null;
+            }break;
+            
+            case SYMS_CvBasicType_HRESULT:
+            {
+              result.kind = SYMS_TypeKind_Void;
+              result.reported_size = 4;
+            }break;
+            
+            case SYMS_CvBasicType_RCHAR:
+            case SYMS_CvBasicType_CHAR:
+            {
+              result.kind = SYMS_TypeKind_Int8;
+              result.mods = SYMS_TypeModifier_Char;
+              result.reported_size = 1;
+            }break;
+            
+            case SYMS_CvBasicType_UCHAR:
+            {
+              result.kind = SYMS_TypeKind_UInt8;
+              result.mods = SYMS_TypeModifier_Char;
+              result.reported_size = 1;
+            }break;
+            
+            case SYMS_CvBasicType_WCHAR:
+            {
+              result.kind = SYMS_TypeKind_UInt16;
+              result.mods = SYMS_TypeModifier_Char;
+              result.reported_size = 2;
+            }break;
+            
+            case SYMS_CvBasicType_BOOL8:
+            case SYMS_CvBasicType_CHAR8:
+            case SYMS_CvBasicType_INT8:
+            {
+              result.kind = SYMS_TypeKind_Int8;
+              result.reported_size = 1;
+            }break;
+            
+            case SYMS_CvBasicType_BOOL16:
+            case SYMS_CvBasicType_CHAR16:
+            case SYMS_CvBasicType_SHORT:
+            case SYMS_CvBasicType_INT16:
+            {
+              result.kind = SYMS_TypeKind_Int16;
+              result.reported_size = 2;
+            }break;
+            
+            case SYMS_CvBasicType_BOOL32:
+            case SYMS_CvBasicType_CHAR32:
+            case SYMS_CvBasicType_INT32:
+            {
+              result.kind = SYMS_TypeKind_Int32;
+              result.reported_size = 4;
+            }break;
+            
+            case SYMS_CvBasicType_BOOL64:
+            case SYMS_CvBasicType_QUAD:
+            case SYMS_CvBasicType_INT64:
+            {
+              result.kind = SYMS_TypeKind_Int64;
+              result.reported_size = 8;
+            }break;
+            
+            case SYMS_CvBasicType_OCT:
+            case SYMS_CvBasicType_INT128:
+            {
+              result.kind = SYMS_TypeKind_Int128;
+              result.reported_size = 16;
+            }break;
+            
+            case SYMS_CvBasicType_UINT8:
+            {
+              result.kind = SYMS_TypeKind_UInt8;
+              result.reported_size = 1;
+            }break;
+            
+            case SYMS_CvBasicType_USHORT:
+            case SYMS_CvBasicType_UINT16:
+            {
+              result.kind = SYMS_TypeKind_UInt16;
+              result.reported_size = 2;
+            }break;
+            
+            case SYMS_CvBasicType_LONG:
+            {
+              result.kind = SYMS_TypeKind_Int32;
+              result.reported_size = 4;
+            }break;
+            
+            case SYMS_CvBasicType_ULONG:
+            {
+              result.kind = SYMS_TypeKind_UInt32;
+              result.reported_size = 4;
+            }break;
+            
+            case SYMS_CvBasicType_UINT32:
+            {
+              result.kind = SYMS_TypeKind_UInt32;
+              result.reported_size = 4;
+            }break;
+            
+            case SYMS_CvBasicType_UQUAD:
+            case SYMS_CvBasicType_UINT64:
+            {
+              result.kind = SYMS_TypeKind_UInt64;
+              result.reported_size = 8;
+            }break;
+            
+            case SYMS_CvBasicType_UOCT:
+            case SYMS_CvBasicType_UINT128:
+            {
+              result.kind = SYMS_TypeKind_UInt128;
+              result.reported_size = 16;
+            }break;
+            
+            case SYMS_CvBasicType_FLOAT16:
+            {
+              result.kind = SYMS_TypeKind_Float16;
+              result.reported_size = 2;
+            }break;
+            
+            case SYMS_CvBasicType_FLOAT32:
+            {
+              result.kind = SYMS_TypeKind_Float32;
+              result.reported_size = 4;
+            }break;
+            
+            case SYMS_CvBasicType_FLOAT64:
+            {
+              result.kind = SYMS_TypeKind_Float64;
+              result.reported_size = 8;
+            }break;
+            
+            case SYMS_CvBasicType_FLOAT32PP:
+            {
+              result.kind = SYMS_TypeKind_Float32PP;
+              result.reported_size = 4;
+            }break;
+            
+            case SYMS_CvBasicType_FLOAT80:
+            {
+              result.kind = SYMS_TypeKind_Float80;
+              result.reported_size = 10;
+            }break;
+            
+            case SYMS_CvBasicType_FLOAT128:
+            {
+              result.kind = SYMS_TypeKind_Float128;
+              result.reported_size = 16;
+            }break;
+            
+            case SYMS_CvBasicType_COMPLEX32:
+            {
+              result.kind = SYMS_TypeKind_Complex32;
+              result.reported_size = 8;
+            }break;
+            
+            case SYMS_CvBasicType_COMPLEX64:
+            {
+              result.kind = SYMS_TypeKind_Complex64;
+              result.reported_size = 16;
+            }break;
+            
+            case SYMS_CvBasicType_COMPLEX80:
+            {
+              result.kind = SYMS_TypeKind_Complex80;
+              result.reported_size = 20;
+            }break;
+            
+            case SYMS_CvBasicType_COMPLEX128:
+            {
+              result.kind = SYMS_TypeKind_Complex128;
+              result.reported_size = 32;
+            }break;
+          }
+        }break;
+        
+        case SYMS_CvBasicPointerKind_16BIT:
+        case SYMS_CvBasicPointerKind_FAR_16BIT:
+        case SYMS_CvBasicPointerKind_HUGE_16BIT:
+        {
+          result.kind = SYMS_TypeKind_Ptr;
+          result.reported_size = 2;
+          result.direct_type.uid = SYMS_PdbPseudoUnit_TPI;
+          result.direct_type.sid = SYMS_ID_u32_u32(SYMS_PdbSymbolIDKind_Index, itype_kind);
+        }break;
+        
+        case SYMS_CvBasicPointerKind_32BIT:
+        case SYMS_CvBasicPointerKind_16_32BIT:
+        {
+          result.kind = SYMS_TypeKind_Ptr;
+          result.reported_size = 4;
+          result.direct_type.uid = SYMS_PdbPseudoUnit_TPI;
+          result.direct_type.sid = SYMS_ID_u32_u32(SYMS_PdbSymbolIDKind_Index, itype_kind);
+        }break;
+        
+        case SYMS_CvBasicPointerKind_64BIT:
+        {
+          result.kind = SYMS_TypeKind_Ptr;
+          result.reported_size  = 8;
+          result.direct_type.uid = SYMS_PdbPseudoUnit_TPI;
+          result.direct_type.sid = SYMS_ID_u32_u32(SYMS_PdbSymbolIDKind_Index, itype_kind);
+        }break;
+      }
+    }
+  }
+  
+  // recorded type info
+  if (resolve.stub != 0){
+    // shared fields for user data type paths
+    SYMS_TypeKind type_kind = SYMS_TypeKind_Null;
+    SYMS_CvTypeProps props = 0;
+    
+    switch (resolve.leaf_kind){
+      default:break;
+      
+      case SYMS_CvLeaf_MODIFIER:
+      {
+        SYMS_CvLeafModifier modifier = {0};
+        if (syms_msf_read_struct_in_range(data, msf, resolve.leaf_range, 0, &modifier)){
+          result.kind = SYMS_TypeKind_Modifier;
+          result.mods = syms_pdb_modifier_from_cv_modifier_flags(modifier.flags);
+          result.direct_type.uid = SYMS_PdbPseudoUnit_TPI;
+          result.direct_type.sid = SYMS_ID_u32_u32(SYMS_PdbSymbolIDKind_Index, modifier.itype);
+          result.reported_size_interp = SYMS_SizeInterpretation_Multiplier;
+          result.reported_size = 1;
+        }
+      }break;
+      
+      case SYMS_CvLeaf_POINTER:
+      {
+        SYMS_CvLeafPointer ptr = {0};
+        if (syms_msf_read_struct_in_range(data, msf, resolve.leaf_range, 0, &ptr)){
+          SYMS_U64 size = SYMS_CvPointerAttribs_Extract_SIZE(ptr.attr);
+          SYMS_CvPointerKind ptr_kind = SYMS_CvPointerAttribs_Extract_KIND(ptr.attr);
+          SYMS_CvPointerMode ptr_mode = SYMS_CvPointerAttribs_Extract_MODE(ptr.attr);
+          
+          SYMS_TypeKind type_kind = syms_pdb_type_kind_from_cv_pointer_mode(ptr_mode);
+          
+          SYMS_CvTypeIndex containing_type = 0;
+          if (type_kind == SYMS_TypeKind_MemberPtr){
+            syms_msf_read_struct_in_range(data, msf, resolve.leaf_range, sizeof(ptr), &containing_type);
+          }
+          
+          result.direct_type.uid = SYMS_PdbPseudoUnit_TPI;
+          result.direct_type.sid = SYMS_ID_u32_u32(SYMS_PdbSymbolIDKind_Index, ptr.itype);
+          result.kind = type_kind;
+          result.mods = syms_pdb_modifier_from_cv_pointer_attribs(ptr.attr);
+          result.reported_size_interp = SYMS_SizeInterpretation_ByteCount;
+          result.reported_size = size;
+          if (containing_type != 0){
+            result.containing_type.uid = SYMS_PdbPseudoUnit_TPI;
+            result.containing_type.sid = SYMS_ID_u32_u32(SYMS_PdbSymbolIDKind_Index, containing_type);
+          }
+        }
+      }break;
+      
+      case SYMS_CvLeaf_PROCEDURE:
+      {
+        SYMS_CvLeafProcedure proc = {0};
+        if (syms_msf_read_struct_in_range(data, msf, resolve.leaf_range, 0, &proc)){
+          // skipped: funcattr, arg_count, arg_itype
+          result.direct_type.uid = SYMS_PdbPseudoUnit_TPI;
+          result.direct_type.sid = SYMS_ID_u32_u32(SYMS_PdbSymbolIDKind_Index, proc.ret_itype);
+          result.kind = SYMS_TypeKind_Proc;
+          result.call_convention = syms_pdb_call_convention_from_cv_call_kind(proc.call_kind);
+        }
+      }break;
+      
+      case SYMS_CvLeaf_MFUNCTION:
+      {
+        SYMS_CvLeafMFunction mfunc = {0};
+        if (syms_msf_read_struct_in_range(data, msf, resolve.leaf_range, 0, &mfunc)){
+          // skipped: class_itype, this_itype, funcattr, arg_count, arg_itype, thisadjust
+          result.direct_type.uid = SYMS_PdbPseudoUnit_TPI;
+          result.direct_type.sid = SYMS_ID_u32_u32(SYMS_PdbSymbolIDKind_Index, mfunc.ret_itype);
+          result.kind = SYMS_TypeKind_Proc;
+          result.call_convention = syms_pdb_call_convention_from_cv_call_kind(mfunc.call_kind);
+        }
+      }break;
+      
+      case SYMS_CvLeaf_ARRAY:
+      {
+        SYMS_CvLeafArray array = {0};
+        if (syms_msf_read_struct_in_range(data, msf, resolve.leaf_range, 0, &array)){
+          result.direct_type.uid = SYMS_PdbPseudoUnit_TPI;
+          result.direct_type.sid = SYMS_ID_u32_u32(SYMS_PdbSymbolIDKind_Index, array.entry_itype);
+          result.kind = SYMS_TypeKind_Array;
+          result.reported_size_interp = SYMS_SizeInterpretation_ByteCount;
+          result.reported_size = resolve.stub->num;
+        }
+      }break;
+      
+      case SYMS_CvLeaf_BITFIELD:
+      {
+        SYMS_CvLeafBitField bitfield = {0};
+        if (syms_msf_read_struct_in_range(data, msf, resolve.leaf_range, 0, &bitfield)){
+          result.direct_type.uid = SYMS_PdbPseudoUnit_TPI;
+          result.direct_type.sid = SYMS_ID_u32_u32(SYMS_PdbSymbolIDKind_Index, bitfield.itype);
+          result.kind = SYMS_TypeKind_Bitfield;
+          result.reported_size_interp = SYMS_SizeInterpretation_Multiplier;
+          result.reported_size = 1;
+          // TODO(allen): bitfield reporting
+        }
+      }break;
+      
+      // udt
+      {
+        case SYMS_CvLeaf_STRUCTURE:
+        {
+          type_kind = SYMS_TypeKind_Struct;
+          goto read_struct;
+        }
+        case SYMS_CvLeaf_CLASS:
+        case SYMS_CvLeaf_INTERFACE:
+        {
+          type_kind = SYMS_TypeKind_Class;
+          goto read_struct;
+        }
+        read_struct:
+        {
+          SYMS_CvLeafStruct struct_ = {0};
+          syms_msf_read_struct_in_range(data, msf, resolve.leaf_range, 0, &struct_);
+          props = struct_.props;
+          goto fill_result;
+        }
+        
+        case SYMS_CvLeaf_UNION:
+        {
+          type_kind = SYMS_TypeKind_Union;
+          SYMS_CvLeafUnion union_ = {0};
+          syms_msf_read_struct_in_range(data, msf, resolve.leaf_range, 0, &union_);
+          props = union_.props;
+          goto fill_result;
+        }
+        
+        case SYMS_CvLeaf_ENUM:
+        {
+          type_kind = SYMS_TypeKind_Enum;
+          SYMS_CvLeafEnum enum_ = {0};
+          syms_msf_read_struct_in_range(data, msf, resolve.leaf_range, 0, &enum_);
+          props = enum_.props;
+          result.direct_type.uid = SYMS_PdbPseudoUnit_TPI;
+          result.direct_type.sid = SYMS_ID_u32_u32(SYMS_PdbSymbolIDKind_Index, enum_.itype);
+          goto fill_result;
+        }
+        
+        case SYMS_CvLeaf_CLASSPTR:
+        case SYMS_CvLeaf_CLASSPTR2:
+        {
+          type_kind = SYMS_TypeKind_Struct;
+          SYMS_CvLeafClassPtr class_ptr = {0};
+          syms_msf_read_struct_in_range(data, msf, resolve.leaf_range, 0, &class_ptr);
+          props = class_ptr.props;
+          goto fill_result;
+        }
+        
+        fill_result:
+        {
+          SYMS_B32 is_fwdref = !!(props & SYMS_CvTypeProp_FWDREF);
+          if (is_fwdref){
+            result.kind = syms_type_kind_fwd_from_main(type_kind);
+            result.reported_size_interp = SYMS_SizeInterpretation_ResolveForwardReference;
+            result.reported_size = 0;
+          }
+          else{
+            result.kind = type_kind;
+            if (type_kind == SYMS_TypeKind_Enum){
+              result.reported_size_interp = SYMS_SizeInterpretation_Multiplier;
+              result.reported_size = 1;
+            }
+            else{
+              result.reported_size_interp = SYMS_SizeInterpretation_ByteCount;
+              result.reported_size = resolve.stub->num;
+            }
+          }
+        }
+      }break;
+      
+      case SYMS_CvLeaf_ALIAS:
+      {
+        SYMS_CvLeafAlias alias = {0};
+        if (syms_msf_read_struct_in_range(data, msf, resolve.leaf_range, 0, &alias)){
+          result.kind = SYMS_TypeKind_Typedef;
+          result.direct_type.uid = SYMS_PdbPseudoUnit_TPI;
+          result.direct_type.sid = SYMS_ID_u32_u32(SYMS_PdbSymbolIDKind_Index, alias.itype);
+          result.reported_size_interp = SYMS_SizeInterpretation_Multiplier;
+          result.reported_size = 1;
+        }
+      }break;
+    }
+  }
+  
+  // type index
+  if (resolve.is_leaf_id){
+    SYMS_CvTypeIndex ti = SYMS_ID_u32_1(id);
+    SYMS_String8 id_data = {(SYMS_U8*)&ti, sizeof(ti)};
+    
+    SYMS_PdbTpiAccel *ipi = &dbg->ipi;
+    SYMS_MsfRange ipi_range = syms_msf_range_from_sn(msf, ipi->type_sn);
+    
+    if (ipi->bucket_count > 0){
+      
+      // get bucket
+      SYMS_U32 name_hash = syms_pdb_hashV1(id_data);
+      SYMS_U32 bucket_index = name_hash%ipi->bucket_count;
+      
+      // iterate bucket
+      for (SYMS_PdbChain *bucket = ipi->buckets[bucket_index];
+           bucket != 0;
+           bucket = bucket->next){
+        SYMS_CvItemId item_id  = bucket->v;
+        SYMS_U32 item_off = syms_pdb_tpi_off_from_ti(data, msf, ipi, item_id);
+        
+        SYMS_CvElement element = syms_cv_element(data, msf, ipi_range, item_off);
+        switch (element.kind){
+          default:break;
+          
+          case SYMS_CvLeaf_UDT_MOD_SRC_LINE:
+          {
+            // TODO(allen): we have never hit this case before; we're not sure how it works,
+            // and have no way to test it.
+            SYMS_ASSERT_PARANOID(!"not implemented");
+#if 0
+            SYMS_CvLeafModSrcLine mod_src_line = {0};
+            if (syms_msf_read_struct_in_range(data, msf, element.range, 0, &mod_src_line)){
+            }
+#endif
+          }break;
+          
+          case SYMS_CvLeaf_UDT_SRC_LINE:
+          {
+            SYMS_CvLeafUDTSrcLine udt_src_line = {0};
+            if (syms_msf_read_struct_in_range(data, msf, element.range, 0, &udt_src_line)){
+              if (udt_src_line.udt_itype == ti){
+                result.src_coord.file_id = SYMS_ID_u32_u32(SYMS_PdbFileIDKind_IPIStringID, udt_src_line.src);
+                result.src_coord.line = udt_src_line.ln;
+                goto dbl_break;
+              }
+            }
+          }break;
+        }
+      }
+      
+      dbl_break:;
+    }
+  }
+  
+  return(result);
+}
+
+SYMS_API SYMS_ConstInfo
+syms_pdb_leaf_const_info_from_id(SYMS_String8 data, SYMS_PdbDbgAccel *dbg, SYMS_PdbUnitAccel *unit,
+                                 SYMS_SymbolID id){
+  // setup accel
+  SYMS_MsfAccel *msf = dbg->msf;
+  
+  // read id
+  SYMS_PdbLeafResolve resolve = syms_pdb_leaf_resolve_from_id(data, dbg, unit, id);
+  
+  // zero clear result
+  SYMS_ConstInfo result = {SYMS_TypeKind_Null};
+  
+  // return const info
+  if (resolve.stub != 0){
+    switch (resolve.leaf_kind){
+      default:break;
+      
+      case SYMS_CvLeaf_ENUMERATE:
+      {
+        SYMS_CvLeafEnumerate enumerate  = {0};
+        if (syms_msf_read_struct_in_range(data, msf, resolve.leaf_range, 0, &enumerate)){
+          // TODO(allen): attribs: SYMS_CvFieldAttribs;
+          
+          SYMS_U32 num_off = sizeof(SYMS_CvLeafEnumerate);
+          SYMS_CvNumeric num = {SYMS_TypeKind_Null};
+          SYMS_U32 num_size = syms_cv_read_numeric(data, msf, resolve.leaf_range, num_off, &num);
+          (void)num_size;
+          result = syms_pdb_const_info_from_cv_numeric(num);
+        }
+      }break;
+    }
+  }
+  
+  return(result);
+}
+
+// pdb unit helpers
+
+SYMS_API SYMS_PdbStub*
+syms_pdb_stub_from_unit_off(SYMS_PdbUnitAccel *unit, SYMS_U32 off){
+  SYMS_PdbStub *result = 0;
+  if (unit->bucket_count > 0){
+    SYMS_U64 hash = syms_hash_u64(off);
+    SYMS_U32 bucket_index = hash%unit->bucket_count;
+    for (SYMS_PdbStub *stub = unit->buckets[bucket_index];
+         stub != 0;
+         stub = stub->bucket_next){
+      if (stub->off == off){
+        result = stub;
+        break;
+      }
+    }
+  }
+  return(result);
+}
+
+SYMS_API SYMS_PdbStub*
+syms_pdb_stub_from_unit_index(SYMS_PdbUnitAccel *unit, SYMS_U32 index){
+  SYMS_PdbStub *result = 0;
+  if (unit->top_min_index <= index){
+    SYMS_U32 relative_index = index - unit->top_min_index;
+    if (relative_index < unit->ti_count){
+      result = unit->ti_stubs + relative_index;
+    }
+  }
+  return(result);
+}
+
+
+// main api
+
+SYMS_API SYMS_PdbUnitAccel*
+syms_pdb_unit_accel_from_id(SYMS_Arena *arena, SYMS_String8 data, SYMS_PdbDbgAccel *dbg,
+                            SYMS_PdbUnitSetAccel *unit_set, SYMS_UnitID id){
+  // setup result
+  SYMS_PdbUnitAccel *result = (SYMS_PdbUnitAccel*)&syms_format_nil;
+  
+  // try symbols pseudo unit
+  if (id == SYMS_PdbPseudoUnit_SYM){
+    result = syms_pdb_pub_sym_accel_from_dbg(arena, data, dbg);
+  }
+  
+  // try types pseudo unit
+  else if (id == SYMS_PdbPseudoUnit_TPI){
+    result = syms_pdb_leaf_accel_from_dbg(arena, data, dbg);
+  }
+  
+  // try compilation unit
+  else{
+    SYMS_PdbCompUnit *unit = syms_pdb_comp_unit_from_id(unit_set, id);
+    if (unit != 0){
+      SYMS_MsfRange range = syms_pdb_msf_range_from_comp_unit(unit, SYMS_PdbCompUnitRange_Symbols);
+      result = syms_pdb_sym_accel_from_range(arena, data, dbg->msf, range, id);
+    }
+  }
+  
+  return(result);
+}
+
+SYMS_API SYMS_UnitID
+syms_pdb_uid_from_accel(SYMS_PdbUnitAccel *unit){
+  return(unit->uid);
+}
+
+SYMS_API SYMS_SymbolIDArray
+syms_pdb_proc_sid_array_from_unit(SYMS_Arena *arena, SYMS_PdbUnitAccel *unit){
+  SYMS_SymbolIDArray result = {0};
+  if (!unit->leaf_set){
+    
+    //- allocate array
+    SYMS_U64 count = unit->proc_count;
+    SYMS_SymbolID *ids = syms_push_array(arena, SYMS_SymbolID, count);
+    
+    //- fill array
+    SYMS_SymbolID *id_ptr = ids;
+    SYMS_PdbStub **stub_ptr = unit->proc_stubs;
+    SYMS_PdbStub **opl = stub_ptr + count;
+    for (; stub_ptr < opl; id_ptr += 1, stub_ptr += 1){
+      *id_ptr = SYMS_ID_u32_u32(SYMS_PdbSymbolIDKind_Off, (**stub_ptr).off);
+    }
+    
+    //- assemble result
+    result.count = count;
+    result.ids = ids;
+  }
+  return(result);
+}
+
+SYMS_API SYMS_SymbolIDArray
+syms_pdb_var_sid_array_from_unit(SYMS_Arena *arena, SYMS_PdbUnitAccel *unit){
+  SYMS_SymbolIDArray result = {0};
+  if (!unit->leaf_set){
+    
+    //- allocate array
+    SYMS_U64 count = unit->var_count;
+    SYMS_SymbolID *ids = syms_push_array(arena, SYMS_SymbolID, count);
+    
+    //- fill array
+    SYMS_SymbolID *id_ptr = ids;
+    SYMS_PdbStub **stub_ptr = unit->var_stubs;
+    SYMS_PdbStub **opl = stub_ptr + count;
+    for (; stub_ptr < opl; id_ptr += 1, stub_ptr += 1){
+      *id_ptr = SYMS_ID_u32_u32(SYMS_PdbSymbolIDKind_Off, (**stub_ptr).off);
+    }
+    
+    //- assemble result
+    result.count = count;
+    result.ids = ids;
+  }
+  return(result);
+}
+
+SYMS_API SYMS_SymbolIDArray
+syms_pdb_type_sid_array_from_unit(SYMS_Arena *arena, SYMS_PdbUnitAccel *unit){
+  SYMS_SymbolIDArray result = {0};
+  if (unit->leaf_set){
+    
+    //- allocate array
+    SYMS_U64 count = unit->udt_count;
+    SYMS_SymbolID *ids = syms_push_array(arena, SYMS_SymbolID, count);
+    
+    //- offset info
+    SYMS_PdbStub *ti_stubs = unit->ti_stubs;
+    SYMS_U32 top_min_index = (SYMS_U32)unit->top_min_index;
+    
+    //- fill array
+    SYMS_SymbolID *id_ptr = ids;
+    SYMS_PdbStub **stub_ptr = unit->udt_stubs;
+    SYMS_PdbStub **opl = stub_ptr + count;
+    for (; stub_ptr < opl; id_ptr += 1, stub_ptr += 1){
+      SYMS_PdbStub *stub = *stub_ptr;
+      SYMS_U32 index = (SYMS_U32)(stub - ti_stubs) + top_min_index;
+      *id_ptr = SYMS_ID_u32_u32(SYMS_PdbSymbolIDKind_Index, index);
+    }
+    
+    //- assemble result
+    result.count = count;
+    result.ids = ids;
+  }
+  return(result);
+}
+
+SYMS_API SYMS_U64
+syms_pdb_symbol_count_from_unit(SYMS_PdbUnitAccel *unit){
+  SYMS_U64 result = unit->top_count;
+  return(result);
+}
+
+SYMS_API SYMS_SymbolKind
+syms_pdb_symbol_kind_from_sid(SYMS_String8 data, SYMS_PdbDbgAccel *dbg,
+                              SYMS_PdbUnitAccel *unit, SYMS_SymbolID id){
+  // dispatch to sym or leaf
+  SYMS_SymbolKind result = SYMS_SymbolKind_Null;
+  if (unit->leaf_set){
+    result = syms_pdb_leaf_symbol_kind_from_id(data, dbg, unit, id);
+  }
+  else{
+    result = syms_pdb_sym_symbol_kind_from_id(data, dbg, unit, id);
+  }
+  return(result);
+}
+
+SYMS_API SYMS_String8
+syms_pdb_symbol_name_from_sid(SYMS_Arena *arena,SYMS_String8 data,SYMS_PdbDbgAccel *dbg,
+                              SYMS_PdbUnitAccel *unit, SYMS_SymbolID id){
+  // dispatch to sym or leaf
+  SYMS_String8 result = {0};
+  if (unit->leaf_set){
+    result = syms_pdb_leaf_symbol_name_from_id(arena, data, dbg, unit, id);
+  }
+  else{
+    result = syms_pdb_sym_symbol_name_from_id(arena, data, dbg, unit, id);
+  }
+  return(result);
+}
+
+SYMS_API SYMS_TypeInfo
+syms_pdb_type_info_from_id(SYMS_String8 data, SYMS_PdbDbgAccel *dbg, SYMS_PdbUnitAccel *unit, SYMS_SymbolID id){
+  SYMS_TypeInfo result = {SYMS_TypeKind_Null};
+  if (unit->leaf_set){
+    result = syms_pdb_leaf_type_info_from_id(data, dbg, unit, id);
+  }
+  return(result);
+}
+
+SYMS_API SYMS_ConstInfo
+syms_pdb_const_info_from_id(SYMS_String8 data, SYMS_PdbDbgAccel *dbg, SYMS_PdbUnitAccel *unit, SYMS_SymbolID id){
+  SYMS_ConstInfo result = {SYMS_TypeKind_Null};
+  if (unit->leaf_set){
+    result = syms_pdb_leaf_const_info_from_id(data, dbg, unit, id);
+  }
+  return(result);
+}
+
+
+////////////////////////////////
+//~ NOTE(allen): PDB Variable Info
+
+// cv parse
+
+SYMS_API SYMS_USID
+syms_pdb_sym_type_from_var_id(SYMS_String8 data, SYMS_PdbDbgAccel *dbg,
+                              SYMS_PdbUnitAccel *unit, SYMS_SymbolID id){
+  // setup accel
+  SYMS_MsfAccel *msf = dbg->msf;
+  
+  // read id
+  SYMS_MsfRange range = syms_msf_range_from_sn(msf, unit->sn);
+  SYMS_PdbStub *stub = syms_pdb_stub_from_unit_off(unit, SYMS_ID_u32_1(id));
+  
+  // zero clear result
+  SYMS_USID result = {0};
+  
+  // parse symbol
+  if (stub != 0){
+    SYMS_CvElement element = syms_cv_element(data, msf, range, stub->off);
+    switch (element.kind){
+      default:break;
+      
+      case SYMS_CvSymKind_LDATA32:
+      case SYMS_CvSymKind_GDATA32:
+      {
+        SYMS_CvData32 data32 = {0};
+        syms_msf_read_struct_in_range(data, msf, element.range, 0, &data32);
+        
+        result.sid = SYMS_ID_u32_u32(SYMS_PdbSymbolIDKind_Index, data32.itype);
+        result.uid = SYMS_PdbPseudoUnit_TPI;
+      }break;
+      
+      case SYMS_CvSymKind_LOCAL:
+      {
+        SYMS_CvLocal loc = {0};
+        syms_msf_read_struct_in_range(data, msf, element.range, 0, &loc);
+        
+        result.sid = SYMS_ID_u32_u32(SYMS_PdbSymbolIDKind_Index, loc.itype);
+        result.uid = SYMS_PdbPseudoUnit_TPI;
+      }break;
+      
+      case SYMS_CvSymKind_REGREL32:
+      {
+        SYMS_CvRegrel32 regrel32 = {0};
+        syms_msf_read_struct_in_range(data, msf, element.range, 0, &regrel32);
+        
+        result.sid = SYMS_ID_u32_u32(SYMS_PdbSymbolIDKind_Index, regrel32.itype);
+        result.uid = SYMS_PdbPseudoUnit_TPI;
+      }break;
+    }
+  }
+  
+  return(result);
+}
+
+
+SYMS_API SYMS_U64
+syms_pdb_sym_voff_from_var_sid(SYMS_String8 data, SYMS_PdbDbgAccel *dbg, SYMS_PdbUnitAccel *unit,
+                               SYMS_SymbolID sid){
+  // setup accel
+  SYMS_MsfAccel *msf = dbg->msf;
+  
+  // read id
+  SYMS_MsfRange range = syms_msf_range_from_sn(msf, unit->sn);
+  SYMS_PdbStub *stub = syms_pdb_stub_from_unit_off(unit, SYMS_ID_u32_1(sid));
+  
+  // zero clear result
+  SYMS_U64 result = 0;
+  
+  // parse symbol
+  if (stub != 0){
+    SYMS_CvElement element = syms_cv_element(data, msf, range, stub->off);
+    switch (element.kind){
+      default:break;
+      
+      case SYMS_CvSymKind_LDATA32:
+      case SYMS_CvSymKind_GDATA32:
+      {
+        SYMS_CvData32 data32 = {0};
+        syms_msf_read_struct_in_range(data, msf, element.range, 0, &data32);
+        SYMS_CoffSection section = syms_pdb_coff_section(data, dbg, data32.sec);
+        result = section.virt_off + data32.sec_off;
+      }break;
+    }
+  }
+  
+  return(result);
+}
+
+
+// main api
+
+SYMS_API SYMS_USID
+syms_pdb_type_from_var_id(SYMS_String8 data, SYMS_PdbDbgAccel *dbg, SYMS_PdbUnitAccel *unit, SYMS_SymbolID id){
+  SYMS_USID result = {0};
+  if (!unit->leaf_set){
+    result = syms_pdb_sym_type_from_var_id(data, dbg, unit, id);
+  }
+  return(result);
+}
+
+SYMS_API SYMS_U64
+syms_pdb_voff_from_var_sid(SYMS_String8 data, SYMS_PdbDbgAccel *dbg, SYMS_PdbUnitAccel *unit, SYMS_SymbolID sid){
+  SYMS_U64 result = 0;
+  if (!unit->leaf_set){
+    result = syms_pdb_sym_voff_from_var_sid(data, dbg, unit, sid);
+  }
+  return(result);
+}
+
+
+////////////////////////////////
+//~ NOTE(allen): PDB Member Info
+
+// cv parse
+
+SYMS_API void
+syms_pdb__field_list_parse(SYMS_Arena *arena, SYMS_String8 data, SYMS_MsfAccel *msf,
+                           SYMS_PdbUnitAccel *unit, SYMS_U32 index, SYMS_PdbMemStubList *out){
+  // get field list from ti
+  SYMS_PdbStub *list_stub = syms_pdb_stub_from_unit_index(unit, index);
+  
+  // parse loop
+  if (list_stub != 0){
+    // leaf range
+    SYMS_MsfRange range = syms_msf_range_from_sn(msf, unit->sn);
+    
+    // use the stubs from initial parse
+    for (SYMS_PdbStub *field_stub = list_stub->first;
+         field_stub != 0;
+         field_stub = field_stub->sibling_next){
+      // read kind
+      SYMS_U32 lf_off = field_stub->off;
+      SYMS_CvLeaf lf_kind = 0;
+      syms_msf_read_struct_in_range(data, msf, range, lf_off, &lf_kind);
+      
+      // recurse into indexes
+      if (lf_kind == SYMS_CvLeaf_INDEX){
+        // TODO(allen): infinite recursion breaker.
+        SYMS_CvLeafIndex index = {0};
+        if (syms_msf_read_struct_in_range(data, msf, range, lf_off, &index)){
+          syms_pdb__field_list_parse(arena, data, msf, unit, index.itype, out);
+        }
+      }
+      
+      // gather members
+      else{
+        switch (lf_kind){
+          case SYMS_CvLeaf_MEMBER:
+          case SYMS_CvLeaf_STMEMBER:
+          case SYMS_CvLeaf_ONEMETHOD:
+          case SYMS_CvLeaf_NESTTYPE:
+          case SYMS_CvLeaf_NESTTYPEEX:
+          case SYMS_CvLeaf_BCLASS:
+          case SYMS_CvLeaf_VBCLASS:
+          case SYMS_CvLeaf_IVBCLASS:
+          case SYMS_CvLeaf_VFUNCTAB:
+          case SYMS_CvLeaf_VFUNCOFF:
+          {
+            SYMS_PdbMemStubNode *node = syms_push_array(arena, SYMS_PdbMemStubNode, 1);
+            node->name = syms_push_string_copy(arena, field_stub->name);
+            node->num = field_stub->num;
+            node->off = field_stub->off;
+            node->off2 = 0;
+            SYMS_QueuePush(out->first, out->last, node);
+            out->mem_count += 1;
+          }break;
+          
+          case SYMS_CvLeaf_METHOD:
+          {
+            SYMS_U32 off = field_stub->off;
+            // TODO(allen): do we get anything useful from this num?
+            SYMS_U32 num = field_stub->num;
+            SYMS_String8 name = syms_push_string_copy(arena, field_stub->name);
+            
+            SYMS_U32 method_off = lf_off + sizeof(lf_kind);
+            SYMS_CvLeafMethod method = {0};
+            syms_msf_read_struct_in_range(data, msf, range, method_off, &method);
+            
+            SYMS_PdbStub *method_list_stub = syms_pdb_stub_from_unit_index(unit, method.itype_list);
+            
+            if (method_list_stub != 0){
+              SYMS_U32 method_list_off = method_list_stub->off;
+              SYMS_PdbStub *method_stub = method_list_stub->first;
+              for (SYMS_U32 i = 0;
+                   i < method.count && method_stub != 0;
+                   i += 1, method_stub = method_stub->sibling_next){
+                SYMS_PdbMemStubNode *node = syms_push_array(arena, SYMS_PdbMemStubNode, 1);
+                node->name = name;
+                node->num = num;
+                node->off = off;
+                node->off2 = method_stub->off - method_list_off;
+                SYMS_QueuePush(out->first, out->last, node);
+                out->mem_count += 1;
+              }
+            }
+          }break;
+        }
+      }
+    }
+  }
+}
+
+// main api
+
+SYMS_API SYMS_PdbMemsAccel*
+syms_pdb_mems_accel_from_sid(SYMS_Arena *arena, SYMS_String8 data, SYMS_PdbDbgAccel *dbg,
+                             SYMS_PdbUnitAccel *unit, SYMS_SymbolID id){
+  // setup accels
+  SYMS_MsfAccel *msf = dbg->msf;
+  
+  // dispatch to leaf
+  SYMS_PdbMemsAccel *result = (SYMS_PdbMemsAccel*)&syms_format_nil;
+  if (unit->leaf_set){
+    
+    // read id
+    SYMS_PdbStub *stub = 0;
+    if (SYMS_ID_u16_0(id) == SYMS_PdbSymbolIDKind_Index){
+      stub = syms_pdb_stub_from_unit_index(unit, SYMS_ID_u32_1(id));
+    }
+    
+    // extract field list from stub
+    SYMS_CvTypeIndex ti = 0;
+    if (stub != 0){
+      SYMS_MsfRange range = syms_msf_range_from_sn(msf, unit->sn);
+      SYMS_CvElement element = syms_cv_element(data, msf, range, stub->off);
+      switch (element.kind){
+        case SYMS_CvLeaf_CLASS:
+        case SYMS_CvLeaf_STRUCTURE:
+        case SYMS_CvLeaf_INTERFACE:
+        {
+          SYMS_CvLeafStruct struct_ = {0};
+          if (syms_msf_read_struct_in_range(data, msf, element.range, 0, &struct_)){
+            ti = struct_.field;
+          }
+        }break;
+        case SYMS_CvLeaf_UNION:
+        {
+          SYMS_CvLeafUnion union_ = {0};
+          if (syms_msf_read_struct_in_range(data, msf, element.range, 0, &union_)){
+            ti = union_.field;
+          }
+        }break;
+      }
+    }
+    
+    if (ti != 0){
+      // recursively gather members
+      SYMS_PdbMemStubList list = {0};
+      syms_pdb__field_list_parse(arena, data, msf, unit, ti, &list);
+      
+      // construct result
+      SYMS_String8 type_name = syms_push_string_copy(arena, stub->name);
+      SYMS_PdbMemStubNode **members = syms_push_array(arena, SYMS_PdbMemStubNode*, list.mem_count);
+      {
+        SYMS_U64 i = 0;
+        for (SYMS_PdbMemStubNode *node = list.first;
+             node != 0;
+             node = node->next, i += 1){
+          members[i] = node;
+        }
+      }
+      
+      result = syms_push_array(arena, SYMS_PdbMemsAccel, 1);
+      syms_memzero_struct(result);
+      result->format = SYMS_FileFormat_PDB;
+      result->type_name = type_name;
+      result->count = list.mem_count;
+      result->members = members;
+    }
+  }
+  
+  return(result);
+}
+
+SYMS_API SYMS_U64
+syms_pdb_mem_count_from_mems(SYMS_PdbMemsAccel *mems){
+  SYMS_U64 result = mems->count;
+  return(result);
+}
+
+SYMS_API void
+syms_pdb__fill_method_mem_info(SYMS_String8 data, SYMS_MsfAccel *msf, SYMS_MsfRange range, SYMS_U32 vbaseoff_off,
+                               SYMS_CvFieldAttribs attribs, SYMS_String8 type_name, SYMS_String8 name, 
+                               SYMS_MemInfo *out){
+  SYMS_CvMemberAccess access = SYMS_CvFieldAttribs_Extract_ACCESS(attribs);
+  SYMS_MemVisibility visibility = syms_mem_visibility_from_member_access(access);
+  
+  SYMS_U32 vbaseoff = 0;
+  SYMS_B32 is_virtual = syms_false;
+  SYMS_B32 is_static = syms_false;
+  
+  SYMS_CvMethodProp mprop = SYMS_CvFieldAttribs_Extract_MPROP(attribs);
+  switch (mprop){
+    case SYMS_CvMethodProp_INTRO:
+    case SYMS_CvMethodProp_PUREINTRO:
+    {
+      syms_msf_read_struct_in_range(data, msf, range, vbaseoff_off, &vbaseoff);
+    }//fallthrough;
+    case SYMS_CvMethodProp_VIRTUAL:
+    case SYMS_CvMethodProp_PUREVIRT:
+    {
+      is_virtual = syms_true;
+    }break;
+    case SYMS_CvMethodProp_STATIC:
+    {
+      is_static = syms_true;
+    }break;
+  }
+  
+  SYMS_MemKind kind = SYMS_MemKind_Method;
+  if (is_static){
+    kind = SYMS_MemKind_StaticMethod;
+  }
+  
+  SYMS_MemFlags flags = 0;
+  if (is_virtual){
+    flags |= SYMS_MemFlag_Virtual;
+  }
+  
+  SYMS_B32 destructor_flag = syms_false;
+  SYMS_String8 name_destructor_skipped = name;
+  if (name.size > 0 && name.str[0] == '~'){
+    name_destructor_skipped.str += 1;
+    name_destructor_skipped.size -= 1;
+    destructor_flag = syms_true;
+  }
+  
+  SYMS_String8 type_name_trunc = syms_string_trunc_symbol_heuristic(type_name);
+  SYMS_String8 name_trunc = syms_string_trunc_symbol_heuristic(name_destructor_skipped);
+  
+  if (syms_string_match(type_name_trunc, name_trunc, 0)){
+    if (destructor_flag){
+      flags |= SYMS_MemFlag_Destructor;
+    }
+    else{
+      flags |= SYMS_MemFlag_Constructor;
+    }
+  }
+  
+  out->kind = kind;
+  out->visibility = visibility;
+  out->flags = flags;
+  out->virtual_off = vbaseoff;
+}
+
+SYMS_API SYMS_MemInfo
+syms_pdb_mem_info_from_number(SYMS_Arena *arena, SYMS_String8 data, SYMS_PdbDbgAccel *dbg, SYMS_PdbUnitAccel *unit,
+                              SYMS_PdbMemsAccel *mems, SYMS_U64 n){
+  // setup accels
+  SYMS_MsfAccel *msf = dbg->msf;
+  
+  SYMS_MemInfo result = {SYMS_MemKind_Null};
+  if (1 <= n && n <= mems->count){
+    SYMS_U64 index = n - 1;
+    SYMS_PdbMemStubNode *stub = mems->members[index];
+    
+    if (stub != 0){
+      // read kind
+      SYMS_U32 lf_off = stub->off;
+      SYMS_MsfRange range = syms_msf_range_from_sn(msf, unit->sn);
+      SYMS_CvLeaf lf_kind = 0;
+      syms_msf_read_struct_in_range(data, msf, range, lf_off, &lf_kind);
+      
+      SYMS_U32 data_off = lf_off + sizeof(lf_kind);
+      
+      switch (lf_kind){
+        default:break;
+        
+        case SYMS_CvLeaf_MEMBER:
+        {
+          SYMS_CvLeafMember member = {0};
+          syms_msf_read_struct_in_range(data, msf, range, data_off, &member);
+          
+          SYMS_CvMemberAccess access = SYMS_CvFieldAttribs_Extract_ACCESS(member.attribs);
+          SYMS_MemVisibility visibility = syms_mem_visibility_from_member_access(access);
+          
+          result.kind = SYMS_MemKind_DataField;
+          result.name = syms_push_string_copy(arena, stub->name);
+          result.visibility = visibility;
+          result.off = stub->num;
+        }break;
+        
+        case SYMS_CvLeaf_STMEMBER:
+        {
+          SYMS_CvLeafStMember stmember = {0};
+          syms_msf_read_struct_in_range(data, msf, range, data_off, &stmember);
+          
+          SYMS_CvMemberAccess access = SYMS_CvFieldAttribs_Extract_ACCESS(stmember.attribs);
+          SYMS_MemVisibility visibility = syms_mem_visibility_from_member_access(access);
+          
+          result.kind = SYMS_MemKind_StaticData;
+          result.name = syms_push_string_copy(arena, stub->name);
+          result.visibility = visibility;
+        }break;
+        
+        case SYMS_CvLeaf_METHOD:
+        {
+          SYMS_CvLeafMethod method = {0};
+          syms_msf_read_struct_in_range(data, msf, range, data_off, &method);
+          
+          SYMS_PdbStub *method_list_stub = syms_pdb_stub_from_unit_index(unit, method.itype_list);
+          
+          if (method_list_stub != 0){
+            SYMS_U32 rec_off = method_list_stub->off + stub->off2;
+            SYMS_CvMethod methodrec = {0};
+            syms_msf_read_struct_in_range(data, msf, range, rec_off, &methodrec);
+            syms_pdb__fill_method_mem_info(data, msf, range, rec_off + sizeof(methodrec), methodrec.attribs,
+                                           mems->type_name, stub->name, &result);
+            result.name = syms_push_string_copy(arena, stub->name);
+          }
+        }break;
+        
+        case SYMS_CvLeaf_ONEMETHOD:
+        {
+          SYMS_CvLeafOneMethod onemethod = {0};
+          syms_msf_read_struct_in_range(data, msf, range, data_off, &onemethod);
+          
+          syms_pdb__fill_method_mem_info(data, msf, range, data_off + sizeof(onemethod), onemethod.attribs,
+                                         mems->type_name, stub->name, &result);
+          result.name = syms_push_string_copy(arena, stub->name);
+        }break;
+        
+        case SYMS_CvLeaf_NESTTYPE:
+        case SYMS_CvLeaf_NESTTYPEEX:
+        {
+          SYMS_MemVisibility visibility = SYMS_MemVisibility_Public;
+          if (lf_kind == SYMS_CvLeaf_NESTTYPEEX){
+            SYMS_CvLeafNestTypeEx nest_type_ex = {0};
+            syms_msf_read_struct_in_range(data, msf, range, data_off, &nest_type_ex);
+            SYMS_CvMemberAccess access = SYMS_CvFieldAttribs_Extract_ACCESS(nest_type_ex.attribs);
+            visibility = syms_mem_visibility_from_member_access(access);
+          }
+          
+          result.kind = SYMS_MemKind_NestedType;
+          result.name = syms_push_string_copy(arena, stub->name);
+          result.visibility = visibility;
+        }break;
+        
+        case SYMS_CvLeaf_BCLASS:
+        {
+          SYMS_CvLeafBClass bclass = {0};
+          syms_msf_read_struct_in_range(data, msf, range, data_off, &bclass);
+          
+          SYMS_CvMemberAccess access = SYMS_CvFieldAttribs_Extract_ACCESS(bclass.attribs);
+          SYMS_MemVisibility visibility = syms_mem_visibility_from_member_access(access);
+          
+          result.kind = SYMS_MemKind_BaseClass;
+          result.name = syms_push_string_copy(arena, stub->name);
+          result.visibility = visibility;
+          result.off = stub->num;
+        }break;
+        
+        case SYMS_CvLeaf_VBCLASS:
+        case SYMS_CvLeaf_IVBCLASS:
+        {
+          SYMS_CvLeafVBClass vbclass = {0};
+          syms_msf_read_struct_in_range(data, msf, range, data_off, &vbclass);
+          
+          SYMS_CvMemberAccess access = SYMS_CvFieldAttribs_Extract_ACCESS(vbclass.attribs);
+          SYMS_MemVisibility visibility = syms_mem_visibility_from_member_access(access);
+          
+          result.kind = SYMS_MemKind_VBaseClassPtr;
+          result.name = syms_push_string_copy(arena, stub->name);
+          result.visibility = visibility;
+          //result.is_virtual = syms_true;
+          result.off = stub->num;
+        }break;
+        
+        case SYMS_CvLeaf_VFUNCTAB:
+        {
+          SYMS_CvLeafVFuncTab vfunctab = {0};
+          syms_msf_read_struct_in_range(data, msf, range, data_off, &vfunctab);
+          
+          result.kind = SYMS_MemKind_VTablePtr;
+          result.off = 0;
+        }break;
+        
+        case SYMS_CvLeaf_VFUNCOFF:
+        {
+          SYMS_CvLeafVFuncOff vfuncoff = {0};
+          syms_msf_read_struct_in_range(data, msf, range, data_off, &vfuncoff);
+          
+          result.kind = SYMS_MemKind_VTablePtr;
+          result.off = vfuncoff.off;
+        }break;
+      }
+    }
+  }
+  
+  return(result);
+}
+
+SYMS_API SYMS_USID
+syms_pdb_type_from_mem_number(SYMS_String8 data, SYMS_PdbDbgAccel *dbg, SYMS_PdbUnitAccel *unit,
+                              SYMS_PdbMemsAccel *mems, SYMS_U64 n){
+  // setup accels
+  SYMS_MsfAccel *msf = dbg->msf;
+  
+  SYMS_USID result = {0};
+  if (1 <= n && n <= mems->count){
+    SYMS_U64 index = n - 1;
+    SYMS_PdbMemStubNode *stub = mems->members[index];
+    
+    if (stub != 0){
+      SYMS_U32 lf_off = stub->off;
+      SYMS_MsfRange range = syms_msf_range_from_sn(msf, unit->sn);
+      SYMS_CvLeaf lf_kind = 0;
+      syms_msf_read_struct_in_range(data, msf, range, lf_off, &lf_kind);
+      
+      SYMS_U32 data_off = lf_off + sizeof(lf_kind);
+      
+      result.uid = SYMS_PdbPseudoUnit_TPI;
+      
+      switch (lf_kind){
+        default:
+        {
+          result.uid = 0;
+        }break;
+        
+        case SYMS_CvLeaf_MEMBER:
+        {
+          SYMS_CvLeafMember member = {0};
+          syms_msf_read_struct_in_range(data, msf, range, data_off, &member);
+          
+          result.sid = SYMS_ID_u32_u32(SYMS_PdbSymbolIDKind_Index, member.itype);
+        }break;
+        
+        case SYMS_CvLeaf_STMEMBER:
+        {
+          SYMS_CvLeafStMember stmember = {0};
+          syms_msf_read_struct_in_range(data, msf, range, data_off, &stmember);
+          result.sid = SYMS_ID_u32_u32(SYMS_PdbSymbolIDKind_Index, stmember.itype);
+        }break;
+        
+        case SYMS_CvLeaf_NESTTYPE:
+        case SYMS_CvLeaf_NESTTYPEEX:
+        {
+          SYMS_CvLeafNestTypeEx nest_type_ex = {0};
+          syms_msf_read_struct_in_range(data, msf, range, data_off, &nest_type_ex);
+          
+          result.sid = SYMS_ID_u32_u32(SYMS_PdbSymbolIDKind_Index, nest_type_ex.itype);
+        }break;
+        
+        case SYMS_CvLeaf_BCLASS:
+        {
+          SYMS_CvLeafBClass bclass = {0};
+          syms_msf_read_struct_in_range(data, msf, range, data_off, &bclass);
+          
+          result.sid = SYMS_ID_u32_u32(SYMS_PdbSymbolIDKind_Index, bclass.itype);
+        }break;
+        
+        case SYMS_CvLeaf_VBCLASS:
+        case SYMS_CvLeaf_IVBCLASS:
+        {
+          SYMS_CvLeafVBClass vbclass = {0};
+          syms_msf_read_struct_in_range(data, msf, range, data_off, &vbclass);
+          
+          result.sid = SYMS_ID_u32_u32(SYMS_PdbSymbolIDKind_Index, vbclass.itype);
+        }break;
+      }
+    }
+  }
+  
+  return(result);
+}
+
+SYMS_API SYMS_SigInfo
+syms_pdb_sig_info_from_mem_number(SYMS_Arena *arena, SYMS_String8 data,
+                                  SYMS_PdbDbgAccel *dbg, SYMS_PdbUnitAccel *unit,
+                                  SYMS_PdbMemsAccel *mems, SYMS_U64 n){
+  // setup accels
+  SYMS_MsfAccel *msf = dbg->msf;
+  
+  SYMS_SigInfo result = {0};
+  if (1 <= n && n <= mems->count){
+    SYMS_U64 index = n - 1;
+    SYMS_PdbMemStubNode *stub = mems->members[index];
+    
+    if (stub != 0){
+      // read kind
+      SYMS_U32 lf_off = stub->off;
+      SYMS_MsfRange range = syms_msf_range_from_sn(msf, unit->sn);
+      SYMS_CvLeaf lf_kind = 0;
+      syms_msf_read_struct_in_range(data, msf, range, lf_off, &lf_kind);
+      
+      SYMS_U32 data_off = lf_off + sizeof(lf_kind);
+      
+      switch (lf_kind){
+        default:
+        {}break;
+        
+        case SYMS_CvLeaf_METHOD:
+        {
+          SYMS_CvLeafMethod method = {0};
+          syms_msf_read_struct_in_range(data, msf, range, data_off, &method);
+          
+          SYMS_PdbStub *method_list_stub = syms_pdb_stub_from_unit_index(unit, method.itype_list);
+          if (method_list_stub != 0){
+            SYMS_U32 rec_off = method_list_stub->off + stub->off2;
+            SYMS_CvMethod methodrec = {0};
+            syms_msf_read_struct_in_range(data, msf, range, rec_off, &methodrec);
+            result = syms_pdb_sig_info_from_sig_index(arena, data, dbg, unit, methodrec.itype);
+          }
+        }break;
+        
+        case SYMS_CvLeaf_ONEMETHOD:
+        {
+          SYMS_CvLeafOneMethod onemethod = {0};
+          syms_msf_read_struct_in_range(data, msf, range, data_off, &onemethod);
+          result = syms_pdb_sig_info_from_sig_index(arena, data, dbg, unit, onemethod.itype);
+        }break;
+      }
+    }
+  }
+  
+  return(result);
+}
+
+SYMS_API SYMS_EnumInfoArray
+syms_pdb_enum_info_array_from_sid(SYMS_Arena *arena, SYMS_String8 data, SYMS_PdbDbgAccel *dbg,
+                                  SYMS_PdbUnitAccel *unit, SYMS_SymbolID sid){
+  // setup accels
+  SYMS_MsfAccel *msf = dbg->msf;
+  
+  // dispatch to leaf
+  SYMS_EnumInfoArray result = {0};
+  if (unit->leaf_set){
+    SYMS_MsfRange range = syms_msf_range_from_sn(msf, unit->sn);
+    
+    // read id
+    SYMS_PdbStub *stub = 0;
+    if (SYMS_ID_u16_0(sid) == SYMS_PdbSymbolIDKind_Index){
+      stub = syms_pdb_stub_from_unit_index(unit, SYMS_ID_u32_1(sid));
+    }
+    
+    // extract field list from stub
+    SYMS_CvTypeIndex ti = 0;
+    if (stub != 0){
+      SYMS_CvElement element = syms_cv_element(data, msf, range, stub->off);
+      if (element.kind == SYMS_CvLeaf_ENUM){
+        SYMS_CvLeafEnum enum_ = {0};
+        if (syms_msf_read_struct_in_range(data, msf, element.range, 0, &enum_)){
+          ti = enum_.field;
+        }
+      }
+    }
+    
+    // get field list from ti
+    SYMS_PdbStub *list_stub = syms_pdb_stub_from_unit_index(unit, ti);
+    
+    // parse loop
+    if (list_stub != 0){
+      // count the enumerates
+      SYMS_U64 count = 0;
+      for (SYMS_PdbStub *field_stub = list_stub->first;
+           field_stub != 0;
+           field_stub = field_stub->sibling_next){
+        SYMS_U32 lf_off = field_stub->off;
+        SYMS_CvLeaf lf_kind = 0;
+        syms_msf_read_struct_in_range(data, msf, range, lf_off, &lf_kind);
+        if (lf_kind == SYMS_CvLeaf_ENUMERATE){
+          count += 1;
+        }
+      }
+      
+      // fill the enumerates array
+      result.count = count;
+      result.enum_info = syms_push_array(arena, SYMS_EnumInfo, count);
+      
+      SYMS_EnumInfo *enum_info_ptr = result.enum_info;
+      for (SYMS_PdbStub *field_stub = list_stub->first;
+           field_stub != 0;
+           field_stub = field_stub->sibling_next){
+        SYMS_U32 lf_off = field_stub->off;
+        SYMS_CvLeaf lf_kind = 0;
+        syms_msf_read_struct_in_range(data, msf, range, lf_off, &lf_kind);
+        if (lf_kind == SYMS_CvLeaf_ENUMERATE){
+          SYMS_U32 num_off = lf_off + sizeof(SYMS_CvLeafEnumerate);
+          SYMS_CvNumeric num = {SYMS_TypeKind_Null};
+          SYMS_U32 num_size = syms_cv_read_numeric(data, msf, range, num_off, &num);
+          SYMS_U32 size_clamped = SYMS_ClampTop(num_size, sizeof(enum_info_ptr->val));
+          
+          enum_info_ptr->name = syms_push_string_copy(arena, field_stub->name);
+          syms_memmove(&enum_info_ptr->val, num.data, size_clamped);
+          enum_info_ptr += 1;
+        }
+      }
+    }
+  }
+  
+  return(result);
+}
+
+
+////////////////////////////////
+//~ NOTE(allen): PDB Procedure Info
+
+// main api
+
+SYMS_API SYMS_UnitIDAndSig
+syms_pdb_proc_sig_handle_from_id(SYMS_String8 data, SYMS_PdbDbgAccel *dbg, SYMS_PdbUnitAccel *unit,
+                                 SYMS_SymbolID id){
+  SYMS_UnitIDAndSig result = {0};
+  
+  if (!unit->leaf_set){
+    // setup accels
+    SYMS_MsfAccel *msf = dbg->msf;
+    
+    // read id
+    SYMS_MsfRange range = syms_msf_range_from_sn(msf, unit->sn);
+    SYMS_PdbStub *stub = syms_pdb_stub_from_unit_off(unit, SYMS_ID_u32_1(id));
+    
+    // read proc
+    if (stub != 0){
+      SYMS_CvElement element = syms_cv_element(data, msf, range, stub->off);
+      switch (element.kind){
+        default:break;
+        
+        case SYMS_CvSymKind_LPROC32:
+        case SYMS_CvSymKind_GPROC32:
+        {
+          SYMS_CvProc32 proc32 = {0};
+          if (syms_msf_read_struct_in_range(data, msf, element.range, 0, &proc32)){
+            result.uid = SYMS_PdbPseudoUnit_TPI;
+            result.sig.v = proc32.itype;
+          }
+        }break;
+      }
+    }
+  }
+  
+  return(result);
+}
+
+SYMS_API SYMS_SigInfo
+syms_pdb_sig_info_from_handle(SYMS_Arena *arena, SYMS_String8 data, SYMS_PdbDbgAccel *dbg, SYMS_PdbUnitAccel *unit,
+                              SYMS_SigHandle handle){
+  SYMS_SigInfo result = {0};
+  if (unit->leaf_set){
+    result = syms_pdb_sig_info_from_sig_index(arena, data, dbg, unit, handle.v);
+  }
+  return(result);
+}
+
+SYMS_API SYMS_U64RangeArray
+syms_pdb_proc_vranges_from_sid(SYMS_Arena *arena, SYMS_String8 data, SYMS_PdbDbgAccel *dbg,
+                               SYMS_PdbUnitAccel *unit, SYMS_SymbolID sid){
+  SYMS_U64RangeArray result = {0};
+  if (!unit->leaf_set){
+    // setup accels
+    SYMS_MsfAccel *msf = dbg->msf;
+    
+    // read id
+    SYMS_MsfRange range = syms_msf_range_from_sn(msf, unit->sn);
+    SYMS_PdbStub *stub = syms_pdb_stub_from_unit_off(unit, SYMS_ID_u32_1(sid));
+    
+    // read proc
+    if (stub != 0){
+      SYMS_CvElement element = syms_cv_element(data, msf, range, stub->off);
+      switch (element.kind){
+        default:break;
+        
+        case SYMS_CvSymKind_LPROC32:
+        case SYMS_CvSymKind_GPROC32:
+        {
+          SYMS_CvProc32 proc32 = {0};
+          if (syms_msf_read_struct_in_range(data, msf, element.range, 0, &proc32)){
+            SYMS_CoffSection section = syms_pdb_coff_section(data, dbg, proc32.sec);
+            
+            SYMS_U64Range range = {0};
+            range.min = section.virt_off + proc32.off;
+            range.max = range.min + proc32.len;
+            
+            result.count = 1;
+            result.ranges = syms_push_array(arena, SYMS_U64Range, 1);
+            result.ranges[0] = range;
+          }
+        }break;
+      }
+    }
+  }
+  
+  return(result);
+}
+
+SYMS_API SYMS_SymbolIDArray
+syms_pdb_scope_children_from_sid(SYMS_Arena *arena, SYMS_String8 data, SYMS_PdbDbgAccel *dbg,
+                                 SYMS_PdbUnitAccel *unit, SYMS_SymbolID id){
+  // dispatch to sym or leaf
+  SYMS_ArenaTemp scratch = syms_get_scratch(&arena, 1);
+  SYMS_SymbolIDList list = {0};
+  if (!unit->leaf_set){
+    // setup accels
+    SYMS_MsfAccel *msf = dbg->msf;
+    
+    // read id
+    SYMS_MsfRange range = syms_msf_range_from_sn(msf, unit->sn);
+    SYMS_PdbStub *stub = syms_pdb_stub_from_unit_off(unit, SYMS_ID_u32_1(id));
+    
+    // build list
+    if (stub != 0){
+      SYMS_CvElement element = syms_cv_element(data, msf, range, stub->off);
+      switch (element.kind){
+        default:break;
+        
+        case SYMS_CvSymKind_LPROC32:
+        case SYMS_CvSymKind_GPROC32:
+        case SYMS_CvSymKind_BLOCK32:
+        {
+          for (SYMS_PdbStub *child = stub->first;
+               child != 0;
+               child = child->sibling_next){
+            SYMS_SymbolIDNode *node = syms_push_array(arena, SYMS_SymbolIDNode, 1);
+            SYMS_QueuePush(list.first, list.last, node);
+            node->id = SYMS_ID_u32_u32(SYMS_PdbSymbolIDKind_Off, child->off);
+            list.count += 1;
+          }
+        }break;
+      }
+    }
+  }
+  
+  // convert list to array
+  SYMS_SymbolIDArray result = syms_sid_array_from_list(arena, &list);
+  
+  syms_release_scratch(scratch);
+  
+  return(result);
+}
+
+
+////////////////////////////////
+//~ NOTE(allen): PDB Signature Info
+
+// pdb specific helper
+SYMS_API SYMS_SigInfo
+syms_pdb_sig_info_from_sig_index(SYMS_Arena *arena, SYMS_String8 data, SYMS_PdbDbgAccel *dbg,
+                                 SYMS_PdbUnitAccel *unit, SYMS_CvTypeIndex index){
+  SYMS_SigInfo result = {0};
+  
+  if (unit->leaf_set){
+    // setup accels
+    SYMS_MsfAccel *msf = dbg->msf;
+    
+    // get unit range
+    SYMS_MsfRange range = syms_msf_range_from_sn(msf, unit->sn);
+    
+    // get member function data
+    SYMS_CvTypeIndex arg_itype = 0;
+    SYMS_CvTypeIndex ret_itype = 0;
+    SYMS_CvTypeIndex this_itype = 0;
+    
+    {
+      SYMS_PdbStub *sig_stub = syms_pdb_stub_from_unit_index(unit, index);
+      if (sig_stub != 0){
+        SYMS_CvElement sig_element = syms_cv_element(data, msf, range, sig_stub->off);
+        switch (sig_element.kind){
+          case SYMS_CvLeaf_PROCEDURE:
+          {
+            SYMS_CvLeafProcedure proc = {0};
+            syms_msf_read_struct_in_range(data, msf, sig_element.range, 0, &proc);
+            arg_itype = proc.arg_itype;
+            ret_itype = proc.ret_itype;
+          }break;
+          
+          case SYMS_CvLeaf_MFUNCTION:
+          {
+            SYMS_CvLeafMFunction mfunc = {0};
+            syms_msf_read_struct_in_range(data, msf, sig_element.range, 0, &mfunc);
+            arg_itype = mfunc.arg_itype;
+            ret_itype = mfunc.ret_itype;
+            this_itype = mfunc.this_itype;
+          }break;
+        }
+      }
+    }
+    
+    // get args list
+    SYMS_CvLeafArgList args_list = {0};
+    SYMS_MsfRange args_list_range = {0};
+    if (arg_itype != 0){
+      SYMS_PdbStub *args_stub = syms_pdb_stub_from_unit_index(unit, arg_itype);
+      if (args_stub != 0){
+        SYMS_CvElement args_element = syms_cv_element(data, msf, range, args_stub->off);
+        if (args_element.kind == SYMS_CvLeaf_ARGLIST){
+          syms_msf_read_struct_in_range(data, msf, args_element.range, 0, &args_list);
+          args_list_range = syms_msf_sub_range(args_element.range, sizeof(args_list),
+                                               args_element.range.size - sizeof(args_list));
+        }
+      }
+    }
+    
+    // build params list
+    
+    SYMS_SymbolIDArray param_array = {0};
+    if (args_list.count > 0){
+      SYMS_ArenaTemp scratch = syms_get_scratch(&arena, 1);
+      SYMS_String8 args_list_data = syms_msf_read_whole_range(scratch.arena, data, msf, args_list_range);
+      
+      SYMS_CvTypeIndex *ti_array = (SYMS_CvTypeIndex*)args_list_data.str;
+      SYMS_U32 data_max_count = args_list_data.size/sizeof(SYMS_CvTypeIndex);
+      SYMS_U32 count = SYMS_ClampTop(args_list.count, data_max_count);
+      
+      SYMS_SymbolID *ids = syms_push_array(arena, SYMS_SymbolID, count);
+      SYMS_SymbolID *id_ptr = ids;
+      SYMS_CvTypeIndex *ti = ti_array;
+      for (SYMS_U32 i = 0; i < count; i += 1, ti += 1, id_ptr += 1){
+        *id_ptr = SYMS_ID_u32_u32(SYMS_PdbSymbolIDKind_Index, *ti);
+      }
+      
+      param_array.ids = ids;
+      param_array.count = count;
+      syms_release_scratch(scratch);
+    }
+    
+    result.uid = SYMS_PdbPseudoUnit_TPI;
+    result.param_type_ids = param_array;
+    result.return_type_id = SYMS_ID_u32_u32(SYMS_PdbSymbolIDKind_Index, ret_itype);
+    result.this_type_id = SYMS_ID_u32_u32(SYMS_PdbSymbolIDKind_Index, this_itype);
+    
+  }
+  
+  return(result);
+}
+
+// main api
+SYMS_API SYMS_SigInfo
+syms_pdb_sig_info_from_id(SYMS_Arena *arena, SYMS_String8 data, SYMS_PdbDbgAccel *dbg,
+                          SYMS_PdbUnitAccel *unit, SYMS_SymbolID id){
+  SYMS_SigInfo result = {0};
+  
+  if (unit->leaf_set){
+    // setup accels
+    SYMS_MsfAccel *msf = dbg->msf;
+    SYMS_MsfRange range = syms_msf_range_from_sn(msf, unit->sn);
+    
+    // method stub symbol -> sig info
+    if (SYMS_ID_u16_0(id) == SYMS_PdbSymbolIDKind_Off){
+      SYMS_PdbStub *stub = syms_pdb_stub_from_unit_off(unit, SYMS_ID_u32_1(id));
+      if (stub != 0){
+        SYMS_U32 lf_off = stub->off;
+        SYMS_CvLeaf lf_kind = 0;
+        syms_msf_read_struct_in_range(data, msf, range, lf_off, &lf_kind);
+        
+        SYMS_U32 data_off = lf_off + sizeof(lf_kind);
+        
+        switch (lf_kind){
+          case SYMS_CvLeaf_METHOD:
+          {
+            SYMS_CvLeafMethod method = {0};
+            syms_msf_read_struct_in_range(data, msf, range, data_off, &method);
+            
+            SYMS_PdbStub *method_list_stub = syms_pdb_stub_from_unit_index(unit, method.itype_list);
+            
+            if (method_list_stub != 0){
+              SYMS_U32 rec_off = method_list_stub->off + SYMS_ID_u16_1(id);
+              SYMS_CvMethod methodrec = {0};
+              syms_msf_read_struct_in_range(data, msf, range, rec_off, &methodrec);
+              result = syms_pdb_sig_info_from_sig_index(arena, data, dbg, unit, methodrec.itype);
+            }
+          }break;
+          
+          case SYMS_CvLeaf_ONEMETHOD:
+          {
+            SYMS_CvLeafOneMethod onemethod = {0};
+            syms_msf_read_struct_in_range(data, msf, range, data_off, &onemethod);
+            result = syms_pdb_sig_info_from_sig_index(arena, data, dbg, unit, onemethod.itype);
+          }break;
+        }
+      }
+    }
+    
+    // type symbol -> sig info
+    else{
+      result = syms_pdb_sig_info_from_sig_index(arena, data, dbg, unit, SYMS_ID_u32_1(id));
+    }
+  }
+  
+  return(result);
+}
+
+////////////////////////////////
+//~ NOTE(allen): PDB Line Info
+
+// pdb specific
+
+SYMS_API SYMS_PdbC13SubSection*
+syms_pdb_c13_sub_sections(SYMS_Arena *arena, SYMS_String8 data, SYMS_MsfAccel *accel, SYMS_MsfRange range){
+  SYMS_PdbC13SubSection *first = 0;
+  SYMS_PdbC13SubSection *last = 0;
+  
+  SYMS_U32 off = 0;
+  for (;;){
+    // read header
+    // TODO(allen): better name for this struct
+    SYMS_CvSubSectionHeader header;
+    syms_memzero_struct(&header);
+    syms_msf_read_struct_in_range(data, accel, range, off, &header);
+    
+    // get sub section info
+    SYMS_U32 sub_section_off = off + sizeof(header);
+    SYMS_U32 sub_section_size = header.size;
+    SYMS_U32 after_sub_section = sub_section_off + sub_section_size;
+    
+    // exit condition
+    if (!syms_msf_bounds_check_in_range(range, after_sub_section)){
+      break;
+    }
+    
+    // emit sub section
+    SYMS_PdbC13SubSection *sub_section = syms_push_array(arena, SYMS_PdbC13SubSection, 1);
+    syms_memzero_struct(sub_section);
+    SYMS_QueuePush(first, last, sub_section);
+    sub_section->kind = header.kind;
+    sub_section->off = sub_section_off;
+    sub_section->size = sub_section_size;
+    
+    // increment off
+    off = sub_section_off + sub_section_size;
+  }
+  
+  return(first);
+}
+
+SYMS_API void
+syms_pdb_loose_push_file_id(SYMS_Arena *arena, SYMS_PdbLineTableLoose *loose, SYMS_FileID id){
+  // check if this id is new
+  SYMS_B32 is_duplicate = syms_false;
+  for (SYMS_PdbFileNode *node = loose->first_file_node;
+       node != 0;
+       node = node->next){
+    SYMS_U64 count = node->count;
+    SYMS_FileID *id_ptr = node->file_ids;
+    for (SYMS_U64 i = 0; i < count; i += 1, id_ptr += 1){
+      if (*id_ptr == id){
+        is_duplicate = syms_true;
+        goto dblbreak;
+      }
+    }
+  }
+  dblbreak:;
+  
+  // insert new id
+  if (!is_duplicate){
+    SYMS_PdbFileNode *last_node = loose->last_file_node;
+    if (last_node != 0 && last_node->count == SYMS_ARRAY_SIZE(last_node->file_ids)){
+      last_node = 0;
+    }
+    if (last_node == 0){
+      last_node = syms_push_array(arena, SYMS_PdbFileNode, 1);
+      SYMS_QueuePush(loose->first_file_node, loose->last_file_node, last_node);
+      last_node->count = 0;
+    }
+    last_node->file_ids[last_node->count] = id;
+    last_node->count += 1;
+    loose->file_count += 1;
+  }
+}
+
+SYMS_API SYMS_Line*
+syms_pdb_loose_push_sequence(SYMS_Arena *arena, SYMS_PdbLineTableLoose *loose, SYMS_U64 line_count){
+  SYMS_PdbLineSequence *new_seq = syms_push_array(arena, SYMS_PdbLineSequence, 1);
+  SYMS_Line *result = syms_push_array(arena, SYMS_Line, line_count);
+  
+  new_seq->lines = result;
+  new_seq->line_count = line_count;
+  
+  SYMS_QueuePush(loose->first_seq, loose->last_seq, new_seq);
+  loose->seq_count += 1;
+  loose->line_count += line_count;
+  
+  return(result);
+}
+
+// main api
+SYMS_API SYMS_String8
+syms_pdb_file_name_from_id(SYMS_Arena *arena, SYMS_String8 data, SYMS_PdbDbgAccel *dbg,
+                           SYMS_PdbUnitSetAccel *unit_set, SYMS_UnitID uid, SYMS_FileID id){
+  // setup accel
+  SYMS_MsfAccel *msf = dbg->msf;
+  
+  // read string
+  SYMS_String8 result = {0};
+  switch (SYMS_ID_u32_0(id)){
+    case SYMS_PdbFileIDKind_IPIOff:
+    {
+      SYMS_PdbTpiAccel *ipi = &dbg->ipi;
+      SYMS_MsfRange ipi_range = syms_msf_range_from_sn(msf, ipi->type_sn);
+      result = syms_msf_read_zstring_in_range(arena, data, msf, ipi_range, SYMS_ID_u32_1(id));
+    }break;
+    
+    case SYMS_PdbFileIDKind_IPIStringID:
+    {
+      SYMS_PdbTpiAccel *ipi = &dbg->ipi;
+      SYMS_MsfRange ipi_range = syms_msf_range_from_sn(msf, ipi->type_sn);
+      SYMS_U32 src_item_off = syms_pdb_tpi_off_from_ti(data, msf, ipi, SYMS_ID_u32_1(id));
+      SYMS_CvElement src_element = syms_cv_element(data, msf, ipi_range, src_item_off);
+      if (src_element.kind == SYMS_CvLeaf_STRING_ID){
+        SYMS_CvLeafStringId string_id = {0};
+        if (syms_msf_read_struct_in_range(data, msf, src_element.range, 0, &string_id)){
+          result = syms_msf_read_zstring_in_range(arena, data, msf, src_element.range, sizeof(string_id));
+        }
+      }
+    }break;
+    
+    case SYMS_PdbFileIDKind_StrTblOff:
+    {
+      result = syms_pdb_strtbl_string_from_off(arena, data, dbg, SYMS_ID_u32_1(id));
+    }break;
+    
+    case SYMS_PdbFileIDKind_C11Off:
+    {
+      SYMS_PdbCompUnit *comp_unit = syms_pdb_comp_unit_from_id(unit_set, uid);
+      SYMS_MsfRange c11_range = syms_pdb_msf_range_from_comp_unit(comp_unit, SYMS_PdbCompUnitRange_C11);
+      SYMS_U32 off = SYMS_ID_u32_1(id);
+      syms_msf_read_in_range(data, msf, c11_range, off, 1, &result.size);
+      result.str = syms_push_array(arena, SYMS_U8, result.size);
+      syms_msf_read_in_range(data, msf, c11_range, off + 1, result.size, result.str);
+    }break;
+  }
+  return(result);
+}
+
+SYMS_API SYMS_LineParseOut
+syms_pdb_line_parse_from_uid(SYMS_Arena *arena, SYMS_String8 data, SYMS_PdbDbgAccel *dbg,
+                             SYMS_PdbUnitSetAccel *unit_set, SYMS_UnitID uid){
+  //- get comp unit
+  SYMS_PdbCompUnit *comp_unit = syms_pdb_comp_unit_from_id(unit_set, uid);
+  
+  //- build table
+  SYMS_LineParseOut result = {0};
+  if (comp_unit != 0){
+    
+    // setup accels
+    SYMS_MsfAccel *msf = dbg->msf;
+    SYMS_ArenaTemp scratch = syms_get_scratch(&arena, 1);
+    
+    // setup loose
+    SYMS_PdbLineTableLoose loose = {0};
+    
+    // try C13
+    SYMS_MsfRange c13_range = syms_pdb_msf_range_from_comp_unit(comp_unit, SYMS_PdbCompUnitRange_C13);
+    if (c13_range.size > 0){
+      SYMS_PdbC13SubSection *sub_sections = syms_pdb_c13_sub_sections(scratch.arena, data, msf, c13_range);
+      
+      // find the checksums section
+      SYMS_PdbC13SubSection *checksums = 0;
+      for (SYMS_PdbC13SubSection *node = sub_sections;
+           node != 0;
+           node = node->next){
+        if (node->kind == SYMS_CvSubSectionKind_FILECHKSMS){
+          checksums = node;
+          break;
+        }
+      }
+      
+      // extract checksums range
+      SYMS_MsfRange checksums_range = {0};
+      if (checksums != 0){
+        checksums_range = syms_msf_sub_range(c13_range, checksums->off, checksums->size);
+      }
+      
+      // iterate lines sections
+      for (SYMS_PdbC13SubSection *sub_sec_node = sub_sections;
+           sub_sec_node != 0;
+           sub_sec_node = sub_sec_node->next){
+        if (sub_sec_node->kind == SYMS_CvSubSectionKind_LINES){
+          SYMS_MsfRange subsec_range = syms_msf_sub_range(c13_range, sub_sec_node->off, sub_sec_node->size);
+          
+          // read header
+          SYMS_CvSubSecLinesHeader subsec_lines_header = {0};
+          syms_msf_read_struct_in_range(data, msf, subsec_range, 0, &subsec_lines_header);
+          
+          // check for columns
+          SYMS_B32 has_columns = ((subsec_lines_header.flags&SYMS_CvSubSecLinesFlag_HasColumns) != 0);
+          
+          // read file
+          SYMS_U32 file_off = sizeof(subsec_lines_header);
+          SYMS_CvFile file = {0};
+          syms_msf_read_struct_in_range(data, msf, subsec_range, file_off, &file);
+          
+          SYMS_U32 line_array_off = file_off + sizeof(file);
+          SYMS_U32 column_array_off = line_array_off + file.num_lines*sizeof(SYMS_CvLine);
+          
+          // track down this file info
+          SYMS_CvChecksum chksum = {0};
+          syms_msf_read_struct_in_range(data, msf, checksums_range, file.file_off, &chksum);
+          
+          // setup the section's file id
+          SYMS_FileID file_id = SYMS_ID_u32_u32(SYMS_PdbFileIDKind_StrTblOff, chksum.name_off);
+          syms_pdb_loose_push_file_id(scratch.arena, &loose, file_id);
+          
+          // get the section's virtual offset
+          SYMS_CoffSection coff_section = syms_pdb_coff_section(data, dbg, subsec_lines_header.sec);
+          SYMS_U32 sequence_base_virt_off = coff_section.virt_off + subsec_lines_header.sec_off;
+          
+          // start the sequence
+          SYMS_Line *lines_out = syms_pdb_loose_push_sequence(scratch.arena, &loose, file.num_lines + 1);
+          
+          // parse info
+          SYMS_Line *line_out_ptr = lines_out;
+          SYMS_U32 line_cursor_off = line_array_off;
+          SYMS_U32 column_cursor_off = column_array_off;
+          for (SYMS_U32 i = 0; i < file.num_lines; i += 1, line_out_ptr += 1){
+            
+            // read line
+            SYMS_CvLine line = {0};
+            syms_msf_read_struct_in_range(data, msf, subsec_range, line_cursor_off, &line);
+            
+            // read column
+            SYMS_CvColumn col = {0};
+            if (has_columns){
+              syms_msf_read_struct_in_range(data, msf, subsec_range, column_cursor_off, &col);
+            }
+            
+            // calculate line data
+            SYMS_U64 line_off = sequence_base_virt_off + line.off;
+            SYMS_U32 line_number = (line.flags>>SYMS_CvLineFlag_LINE_NUMBER_SHIFT)&SYMS_CvLineFlag_LINE_NUMBER_MASK;
+            //SYMS_U32 delta_to_end = (line.flags>>SYMS_CvLineFlag_DELTA_TO_END_SHIFT)&SYMS_CvLineFlag_DELTA_TO_END_MASK;
+            //SYMS_B32 statement = ((line.flags & SYMS_CvLineFlag_STATEMENT) != 0);
+            
+            // emit
+            line_out_ptr->src_coord.file_id = file_id;
+            line_out_ptr->src_coord.line = line_number;
+            line_out_ptr->src_coord.col = col.start;
+            line_out_ptr->voff = line_off;
+            
+            // increment 
+            line_cursor_off += sizeof(line);
+            column_cursor_off += sizeof(col);
+          }
+          
+          // explicitly add an ender
+          {
+            line_out_ptr->src_coord.file_id = file_id;
+            line_out_ptr->src_coord.line = 0;
+            line_out_ptr->src_coord.col = 0;
+            line_out_ptr->voff = sequence_base_virt_off + subsec_lines_header.len;
+          }
+        }
+      }
+    }
+    
+    // try C11
+    SYMS_MsfRange c11_range = syms_pdb_msf_range_from_comp_unit(comp_unit, SYMS_PdbCompUnitRange_C11);
+    if (c11_range.size > 0){
+      // c11 layout:
+      //  file_count: U16;
+      //  range_count: U16;
+      //  file_array: U32[file_count];
+      //  range_array: U64[range_count];
+      //  range_array2: U16[range_count];
+      //  @align(4)
+      //  ... more data ...
+      struct C11Header{
+        SYMS_U16 file_count;
+        SYMS_U16 range_count;
+      };
+      struct C11Header header = {0};
+      syms_msf_read_struct_in_range(data, msf, c11_range, 0, &header);
+      SYMS_U32 unk_file_array_off   = sizeof(header);
+      SYMS_U32 unk_range_array_off  = unk_file_array_off + header.file_count*4;
+      SYMS_U32 unk_range_array2_off = unk_range_array_off + header.range_count*8;
+      SYMS_U32 after_unk_range_array2_off = unk_range_array2_off + header.range_count*2;
+      SYMS_U32 after_header_off = (after_unk_range_array2_off + 3)&~3;
+      
+      // iterate file sections
+      SYMS_U32 file_section_cursor = after_header_off;
+      for (;file_section_cursor < c11_range.size;){
+        // file section layout:
+        //  sec_count: U16;
+        //  padding: U16;
+        //  sec_array: U32[sec_count];
+        //  range_array: U64[sec_count];
+        //  path: length-prefixed with 1 byte;
+        //  @align(4)
+        //  ... more data ...
+        SYMS_U32 file_section_off = file_section_cursor;
+        SYMS_U16 sec_count = 0;
+        syms_msf_read_struct_in_range(data, msf, c11_range, file_section_off, &sec_count);
+        SYMS_U32 unk_sec_array_off = file_section_off + 4;
+        SYMS_U32 range_array_off   = unk_sec_array_off + 4*sec_count;
+        SYMS_U32 parse_path_off    = range_array_off + 8*sec_count;
+        
+        SYMS_U32 path_size = 0;
+        syms_msf_read_in_range(data, msf, c11_range, parse_path_off, 1, &path_size);
+        SYMS_U32 path_off = parse_path_off + 1; (void)path_off;
+        SYMS_U32 after_path_off = parse_path_off + 1 + path_size;
+        
+        SYMS_U32 after_file_section_off = (after_path_off + 3)&~3;
+        
+        // setup section file id
+        SYMS_FileID file_id = SYMS_ID_u32_u32(SYMS_PdbFileIDKind_C11Off, parse_path_off);
+        syms_pdb_loose_push_file_id(scratch.arena, &loose, file_id);
+        
+        // iterate sections
+        SYMS_U32 section_cursor = after_file_section_off;
+        for (SYMS_U16 i = 0; i < sec_count; i += 1){
+          // section layout:
+          //  sec: U16;
+          //  line_count: U16;
+          //  offs_array: U32[line_count];
+          //  nums_array: U16[line_count];
+          //  @align(4)
+          //  ... more data ...
+          SYMS_U32 section_off = section_cursor;
+          struct C11Sec{
+            SYMS_U16 sec;
+            SYMS_U16 line_count;
+          };
+          struct C11Sec c11_sec = {0};
+          syms_msf_read_struct_in_range(data, msf, c11_range, section_off, &c11_sec);
+          SYMS_U32 offs_array_off = section_off + sizeof(c11_sec);
+          SYMS_U32 nums_array_off = offs_array_off + c11_sec.line_count*4;
+          SYMS_U32 after_nums_array_off = nums_array_off + c11_sec.line_count*2;
+          SYMS_U32 after_section_off = (after_nums_array_off + 3)&~3;
+          
+          // get section range
+          SYMS_U32Range sec_range = {0};
+          syms_msf_read_struct_in_range(data, msf, c11_range, range_array_off + 8*i, &sec_range);
+          
+          // get the section's virtual offset
+          SYMS_CoffSection coff_section = syms_pdb_coff_section(data, dbg, c11_sec.sec);
+          SYMS_U32 sequence_base_virt_off = coff_section.virt_off;
+          
+          // start the sequence
+          SYMS_Line *lines_out = syms_pdb_loose_push_sequence(scratch.arena, &loose, c11_sec.line_count + 1);
+          
+          // iterate line info
+          SYMS_Line *line_out_ptr = lines_out;
+          for (SYMS_U16 j = 0; j < c11_sec.line_count; j += 1, line_out_ptr += 1){
+            SYMS_U32 line_off = 0;
+            SYMS_U16 line_num = 0;
+            syms_msf_read_struct_in_range(data, msf, c11_range, offs_array_off + 4*j, &line_off);
+            syms_msf_read_struct_in_range(data, msf, c11_range, nums_array_off + 2*j, &line_num);
+            
+            line_out_ptr->src_coord.file_id = file_id;
+            line_out_ptr->src_coord.line = line_num;
+            line_out_ptr->src_coord.col = 0;
+            line_out_ptr->voff = sequence_base_virt_off + line_off;
+          }
+          
+          // explicitly add an ender
+          {
+            line_out_ptr->src_coord.file_id = file_id;
+            line_out_ptr->src_coord.line = 0;
+            line_out_ptr->src_coord.col = 0;
+            line_out_ptr->voff = sequence_base_virt_off + sec_range.max;
+          }
+          
+          // increment
+          section_cursor = after_section_off;
+        }
+        
+        // increment
+        file_section_cursor = section_cursor;
+      }
+    }
+    
+    // file table from list
+    SYMS_FileID *id_array = syms_push_array(arena, SYMS_FileID, loose.file_count);
+    {
+      SYMS_FileID *file_id_ptr = id_array;
+      for (SYMS_PdbFileNode *node = loose.first_file_node;
+           node != 0;
+           node = node->next){
+        syms_memmove(file_id_ptr, node->file_ids, sizeof(*node->file_ids)*node->count);
+        file_id_ptr += node->count;
+      }
+    }
+    
+    // line table from lists
+    SYMS_U64 *sequence_index_array = syms_push_array(arena, SYMS_U64, loose.seq_count + 1);
+    SYMS_Line *line_array = syms_push_array(arena, SYMS_Line, loose.line_count);
+    {
+      SYMS_U64 *index_ptr = sequence_index_array;
+      SYMS_Line *line_ptr = line_array;
+      SYMS_U64 index = 0;
+      for (SYMS_PdbLineSequence *seq = loose.first_seq;
+           seq != 0;
+           seq = seq->next){
+        SYMS_U64 seq_line_count = seq->line_count;
+        // fill indices
+        *index_ptr = index;
+        index_ptr += 1;
+        // fill lines
+        syms_memmove(line_ptr, seq->lines, sizeof(*seq->lines)*seq_line_count);
+        line_ptr += seq_line_count;
+        index += seq_line_count;
+      }
+      *index_ptr = index;
+    }
+    
+    // fill result
+    result.file_id_array.count = loose.file_count;
+    result.file_id_array.ids = id_array;
+    result.line_table.sequence_count = loose.seq_count;
+    result.line_table.sequence_index_array = sequence_index_array;
+    result.line_table.line_count = loose.line_count;
+    result.line_table.line_array = line_array;
+    
+    syms_release_scratch(scratch);
+  }
+  
+  return(result);
+}
+
+
+////////////////////////////////
+// NOTE(allen): PDB Name Maps
+
+SYMS_API SYMS_PdbMapAccel*
+syms_pdb_type_map_from_dbg(SYMS_Arena *arena, SYMS_String8 data, SYMS_PdbDbgAccel *dbg){
+  SYMS_PdbMapAccel *result = syms_push_array(arena, SYMS_PdbMapAccel, 1);
+  result->format = SYMS_FileFormat_PDB;
+  result->uid = SYMS_PdbPseudoUnit_TPI;
+  return(result);
+}
+
+SYMS_API SYMS_PdbMapAccel*
+syms_pdb_image_symbol_map_from_dbg(SYMS_Arena *arena, SYMS_String8 data, SYMS_PdbDbgAccel *dbg){
+  SYMS_PdbMapAccel *result = syms_push_array(arena, SYMS_PdbMapAccel, 1);
+  result->format = SYMS_FileFormat_PDB;
+  result->uid = SYMS_PdbPseudoUnit_SYM;
+  return(result);
+}
+
+SYMS_API SYMS_UnitID
+syms_pdb_partner_uid_from_map(SYMS_PdbMapAccel *map){
+  return(map->uid);
+}
+
+SYMS_API SYMS_USIDList
+syms_pdb_usid_list_from_string(SYMS_Arena *arena, SYMS_String8 data, SYMS_PdbDbgAccel *dbg,
+                               SYMS_PdbUnitAccel *unit, SYMS_PdbMapAccel *map, SYMS_String8 string){
+  SYMS_USIDList result = {0};
+  if (unit->format == SYMS_FileFormat_PDB){
+    switch (map->uid){
+      case SYMS_PdbPseudoUnit_TPI:
+      {
+        result = syms_pdb_types_from_name(arena, data, dbg, unit, string);
+      }break;
+      
+      case SYMS_PdbPseudoUnit_SYM:
+      {
+        result = syms_pdb_symbols_from_name(arena, data, dbg->msf, &dbg->gsi, unit, string);
+      }break;
+    }
+  }
+  return(result);
+}
+
+
+////////////////////////////////
+// NOTE(allen): PDB CV -> Syms Enums and Flags
+
+SYMS_API SYMS_TypeKind
+syms_pdb_type_kind_from_cv_pointer_mode(SYMS_CvPointerMode mode){
+  SYMS_TypeKind result = SYMS_TypeKind_Null;
+  switch (mode){
+    default:break;
+    case SYMS_CvPointerMode_PTR:    result = SYMS_TypeKind_Ptr;             break;
+    case SYMS_CvPointerMode_LVREF:  result = SYMS_TypeKind_LValueReference; break;
+    case SYMS_CvPointerMode_RVREF:  result = SYMS_TypeKind_RValueReference; break;
+    case SYMS_CvPointerMode_PMEM:
+    case SYMS_CvPointerMode_PMFUNC: result = SYMS_TypeKind_MemberPtr;       break;
+  }
+  return(result);
+}
+
+SYMS_API SYMS_TypeModifiers
+syms_pdb_modifier_from_cv_pointer_attribs(SYMS_CvPointerAttribs attribs){
+  SYMS_TypeModifiers result = 0;
+  if (attribs & SYMS_CvPointerAttrib_IS_VOLATILE){
+    result |= SYMS_TypeModifier_Volatile;
+  }
+  if (attribs & SYMS_CvPointerAttrib_IS_CONST){
+    result |= SYMS_TypeModifier_Const;
+  }
+  if (attribs & SYMS_CvPointerAttrib_IS_UNALIGNED){
+    result |= SYMS_TypeModifier_Packed;
+  }
+  if (attribs & SYMS_CvPointerAttrib_IS_RESTRICTED){
+    result |= SYMS_TypeModifier_Restrict;
+  }
+  if (attribs & SYMS_CvPointerAttrib_IS_LREF){
+    result |= SYMS_TypeModifier_Reference;
+  }
+  if (attribs & SYMS_CvPointerAttrib_IS_RREF){
+    result |= SYMS_TypeModifier_RValueReference;
+  }
+  return(result);
+}
+
+SYMS_API SYMS_TypeModifiers
+syms_pdb_modifier_from_cv_modifier_flags(SYMS_CvModifierFlags flags){
+  SYMS_TypeModifiers result = 0;
+  if (flags & SYMS_CvModifierFlag_VOLATILE){
+    result |= SYMS_TypeModifier_Volatile;
+  }
+  if (flags & SYMS_CvModifierFlag_CONST){
+    result |= SYMS_TypeModifier_Const;
+  }
+  if (flags & SYMS_CvModifierFlag_UNALIGNED){
+    result |= SYMS_TypeModifier_Packed;
+  }
+  return(result);
+}
+
+SYMS_API SYMS_CallConvention
+syms_pdb_call_convention_from_cv_call_kind(SYMS_CvCallKind kind){
+  SYMS_CallConvention result = SYMS_CallConvention_NULL;
+  switch (kind){
+    case SYMS_CvCallKind_NEAR_C:      result = SYMS_CallConvention_NEAR_C;      break;
+    case SYMS_CvCallKind_FAR_C:       result = SYMS_CallConvention_FAR_C;       break;
+    case SYMS_CvCallKind_NEAR_PASCAL: result = SYMS_CallConvention_NEAR_PASCAL; break;
+    case SYMS_CvCallKind_FAR_PASCAL:  result = SYMS_CallConvention_FAR_PASCAL;  break;
+    case SYMS_CvCallKind_NEAR_FAST:   result = SYMS_CallConvention_NEAR_FAST;   break;
+    case SYMS_CvCallKind_FAR_FAST:    result = SYMS_CallConvention_FAR_FAST;    break;
+    case SYMS_CvCallKind_NEAR_STD:    result = SYMS_CallConvention_NEAR_STD;    break;
+    case SYMS_CvCallKind_FAR_STD:     result = SYMS_CallConvention_FAR_STD;     break;
+    case SYMS_CvCallKind_NEAR_SYS:    result = SYMS_CallConvention_NEAR_SYS;    break;
+    case SYMS_CvCallKind_FAR_SYS:     result = SYMS_CallConvention_FAR_SYS;     break;
+    case SYMS_CvCallKind_THISCALL:    result = SYMS_CallConvention_THISCALL;    break;
+    case SYMS_CvCallKind_MIPSCALL:    result = SYMS_CallConvention_MIPSCALL;    break;
+    case SYMS_CvCallKind_GENERIC:     result = SYMS_CallConvention_GENERIC;     break;
+    case SYMS_CvCallKind_ALPHACALL:   result = SYMS_CallConvention_ALPHACALL;   break;
+    case SYMS_CvCallKind_PPCCALL:     result = SYMS_CallConvention_PPCCALL;     break;
+    case SYMS_CvCallKind_SHCALL:      result = SYMS_CallConvention_SHCALL;      break;
+    case SYMS_CvCallKind_ARMCALL:     result = SYMS_CallConvention_ARMCALL;     break;
+    case SYMS_CvCallKind_AM33CALL:    result = SYMS_CallConvention_AM33CALL;    break;
+    case SYMS_CvCallKind_TRICALL:     result = SYMS_CallConvention_TRICALL;     break;
+    case SYMS_CvCallKind_SH5CALL:     result = SYMS_CallConvention_SH5CALL;     break;
+    case SYMS_CvCallKind_M32RCALL:    result = SYMS_CallConvention_M32RCALL;    break;
+    case SYMS_CvCallKind_CLRCALL:     result = SYMS_CallConvention_CLRCALL;     break;
+    case SYMS_CvCallKind_INLINE:      result = SYMS_CallConvention_INLINE;      break;
+    case SYMS_CvCallKind_NEAR_VECTOR: result = SYMS_CallConvention_NEAR_VECTOR; break;
+  }
+  return(result);
+}
+
+#if 0
+SYMS_API SYMS_ProcedureFlags
+syms_pdb_procedure_flags_from_cv_procedure_flags(SYMS_CvProcFlags pdb_flags){
+  SYMS_ProcedureFlags result = 0;
+  if (pdb_flags){
+    
+  }
+  return(result);
+}
+
+SYMS_API SYMS_VarFlags
+syms_var_flags_from_cv_local_flags(SYMS_CvLocalFlags pdb_flags){
+  SYMS_VarFlags result = 0;
+  // TODO(allen): canonical flag conversion
+  if (pdb_flags & SYMS_CvLocalFlag_PARAM){
+    result |= SYMS_VarFlag_Parameter;
+  }
+  if (pdb_flags & SYMS_CvLocalFlag_COMPGEN){
+    result |= SYMS_VarFlag_CompilerGen;
+  }
+  if (pdb_flags & SYMS_CvLocalFlag_ALIASED){
+    result |= SYMS_VarFlag_Aliased;
+  }
+  if (pdb_flags & SYMS_CvLocalFlag_RETVAL){
+    result |= SYMS_VarFlag_ReturnValue;
+  }
+  if (pdb_flags & SYMS_CvLocalFlag_OPTOUT){
+    result |= SYMS_VarFlag_OptOut;
+  }
+  if (pdb_flags & SYMS_CvLocalFlag_STATIC){
+    result |= SYMS_VarFlag_Static;
+  }
+  if (pdb_flags & SYMS_CvLocalFlag_GLOBAL){
+    result |= SYMS_VarFlag_Global;
+  }
+  return(result);
+}
+#endif
+
+#endif //SYMS_PDB_PARSER_C
