Update syms library to latest version, includes dwarf5 support

diff --git a/Engine/Source/Runtime/SymsLib/syms/core/pdb/syms_pdb_parser.h b/Engine/Source/Runtime/SymsLib/syms/core/pdb/syms_pdb_parser.h
new file mode 100644
--- /dev/null
+++ b/Engine/Source/Runtime/SymsLib/syms/core/pdb/syms_pdb_parser.h
@@ -0,0 +1,585 @@
+// Copyright Epic Games, Inc. All Rights Reserved.
+/* date = May 11th 2021 4:59 pm */
+
+#ifndef SYMS_PDB_PARSER_H
+#define SYMS_PDB_PARSER_H
+
+////////////////////////////////
+//~ NOTE(allen): PDB's Prerequisite Formats
+
+#include "syms/core/pdb/syms_msf_parser.h"
+
+////////////////////////////////
+//~ NOTE(allen): PDB Fundamental
+
+// TODO(allen): version of this with optimized memory locality could be pretty useful.
+typedef struct SYMS_PdbChain{
+  struct SYMS_PdbChain *next;
+  SYMS_U32 v;
+} SYMS_PdbChain;
+
+typedef enum SYMS_PdbSymbolIDKind{
+  SYMS_PdbSymbolIDKind_Off,
+  SYMS_PdbSymbolIDKind_Index,
+} SYMS_PdbSymbolIDKind;
+
+typedef enum SYMS_PdbFileIDKind{
+  SYMS_PdbFileIDKind_Null,
+  SYMS_PdbFileIDKind_IPIOff,
+  SYMS_PdbFileIDKind_IPIStringID,
+  SYMS_PdbFileIDKind_StrTblOff,
+  SYMS_PdbFileIDKind_C11Off,
+} SYMS_PdbFileIDKind;
+
+typedef enum SYMS_PdbPseudoUnit{
+  SYMS_PdbPseudoUnit_Null,
+  SYMS_PdbPseudoUnit_SYM,
+  SYMS_PdbPseudoUnit_TPI,
+  SYMS_PdbPseudoUnit_LAST_PSEUDO_UNIT = SYMS_PdbPseudoUnit_TPI,
+  SYMS_PdbPseudoUnit_COUNT = SYMS_PdbPseudoUnit_TPI,
+  SYMS_PdbPseudoUnit_FIRST_COMP_UNIT,
+} SYMS_PdbPseudoUnit;
+
+////////////////////////////////
+//~ NOTE(allen): CV Parsing
+
+typedef struct SYMS_CvElement{
+  SYMS_MsfRange range;
+  SYMS_U32 next_off;
+  SYMS_U16 kind;
+} SYMS_CvElement;
+
+typedef struct SYMS_CvNumeric{
+  SYMS_TypeKind kind;
+  SYMS_U8 data[32];
+} SYMS_CvNumeric;
+
+////////////////////////////////
+//~ NOTE(allen): PDB DBI Information
+
+typedef struct SYMS_PdbDbiAccel{
+  SYMS_B16 valid;
+  SYMS_CoffMachineType machine_type;
+  SYMS_MsfStreamNumber gsi_sn;
+  SYMS_MsfStreamNumber psi_sn;
+  SYMS_MsfStreamNumber sym_sn;
+  
+  SYMS_U64 range_off[(SYMS_U64)(SYMS_PdbDbiRange_COUNT) + 1];
+  SYMS_MsfStreamNumber dbg_sn[SYMS_PdbDbiStream_COUNT];
+} SYMS_PdbDbiAccel;
+
+typedef struct SYMS_PdbInfoSlot{
+  struct SYMS_PdbInfoSlot *next;
+  SYMS_String8 string;
+  SYMS_MsfStreamNumber sn;
+} SYMS_PdbInfoSlot;
+
+typedef struct SYMS_PdbInfoTable{
+  SYMS_PdbInfoSlot *first;
+  SYMS_PdbInfoSlot *last;
+  SYMS_PeGuid auth_guid;
+} SYMS_PdbInfoTable;
+
+typedef struct SYMS_PdbNamedStreamArray{
+  SYMS_MsfStreamNumber sn[SYMS_PdbNamedStream_COUNT];
+} SYMS_PdbNamedStreamArray;
+
+typedef struct SYMS_PdbStrtblAccel{
+  SYMS_U32 bucket_count;
+  SYMS_MsfStreamNumber sn;
+  SYMS_U32Range strblock;
+  SYMS_U32Range buckets;
+} SYMS_PdbStrtblAccel;
+
+////////////////////////////////
+//~ NOTE(allen): PDB Units
+
+typedef struct SYMS_PdbCompUnit{
+  SYMS_MsfStreamNumber sn;
+  SYMS_U64 range_off[(SYMS_U64)(SYMS_PdbCompUnitRange_COUNT) + 1];
+  
+  SYMS_PdbStringIndex src_file;
+  SYMS_PdbStringIndex pdb_file;
+  
+  SYMS_String8 module_name;
+  SYMS_String8 obj_name;
+} SYMS_PdbCompUnit;
+
+typedef struct SYMS_PdbCompUnitNode{
+  struct SYMS_PdbCompUnitNode *next;
+  SYMS_PdbCompUnit comp_unit;
+} SYMS_PdbCompUnitNode;
+
+typedef struct SYMS_PdbUnitSetAccel{
+  SYMS_FileFormat format;
+  SYMS_PdbCompUnit **comp_units;
+  SYMS_U64 comp_count;
+} SYMS_PdbUnitSetAccel;
+
+typedef struct SYMS_PdbStub{
+  struct SYMS_PdbStub *bucket_next;
+  struct SYMS_PdbStub *sibling_next;
+  struct SYMS_PdbStub *first;
+  struct SYMS_PdbStub *last;
+  SYMS_U32 off;
+  SYMS_U32 num;
+  union{
+    SYMS_U32 num2;
+    SYMS_String8 name;
+  };
+} SYMS_PdbStub;
+
+typedef struct SYMS_PdbStubRef{
+  struct SYMS_PdbStubRef *next;
+  SYMS_PdbStub *stub;
+} SYMS_PdbStubRef;
+
+typedef struct SYMS_PdbUnitAccel{
+  SYMS_FileFormat format;
+  SYMS_B32 leaf_set;
+  SYMS_MsfStreamNumber sn;
+  SYMS_PdbStub **top_stubs;
+  SYMS_U64 top_count;
+  SYMS_U64 top_min_index;
+  SYMS_PdbStub **buckets;
+  SYMS_U64 bucket_count;
+  SYMS_U64 all_count;
+  SYMS_PdbStub *ti_stubs;
+  SYMS_U64 ti_count;
+  SYMS_UnitID uid;
+  
+  SYMS_PdbStub **proc_stubs;
+  SYMS_U64 proc_count;
+  SYMS_PdbStub **var_stubs;
+  SYMS_U64 var_count;
+  SYMS_PdbStub **udt_stubs;
+  SYMS_U64 udt_count;
+} SYMS_PdbUnitAccel;
+
+typedef struct SYMS_PdbLeafResolve{
+  SYMS_PdbStub *stub;
+  SYMS_CvLeaf leaf_kind;
+  SYMS_MsfRange leaf_range;
+  SYMS_B32 is_leaf_id;
+} SYMS_PdbLeafResolve;
+
+////////////////////////////////
+//~ NOTE(allen): PDB C13
+
+typedef struct SYMS_PdbC13SubSection{
+  struct SYMS_PdbC13SubSection *next;
+  SYMS_CvSubSectionKind kind;
+  SYMS_U32 off;
+  SYMS_U32 size;
+} SYMS_PdbC13SubSection;
+
+////////////////////////////////
+//~ NOTE(allen): PDB TPI
+
+typedef struct SYMS_PdbTpiOffRange{
+  SYMS_U32 first_off;
+  SYMS_CvTypeIndex first_ti;
+  SYMS_CvTypeIndex opl_ti;
+} SYMS_PdbTpiOffRange;
+
+typedef struct SYMS_PdbTpiAccel{
+  SYMS_MsfStreamNumber type_sn;
+  SYMS_PdbTpiOffHint *hints;
+  SYMS_PdbChain **buckets;
+  SYMS_U32 bucket_count;
+  SYMS_U64 count;
+  SYMS_U32 base_off;
+  SYMS_CvTypeIndex first_ti;
+  SYMS_CvTypeIndex opl_ti;
+  SYMS_U32 *off;
+} SYMS_PdbTpiAccel;
+
+////////////////////////////////
+//~ NOTE(allen): PDB GSI
+
+typedef struct SYMS_PdbGsiAccel{
+  SYMS_PdbChain **buckets;
+  SYMS_U32 bucket_count;
+} SYMS_PdbGsiAccel;
+
+////////////////////////////////
+//~ NOTE(allen): PDB Accel Types
+
+typedef struct SYMS_PdbFileAccel{
+  SYMS_FileFormat format;
+  SYMS_MsfAccel *msf;
+} SYMS_PdbFileAccel;
+
+typedef struct SYMS_PdbDbgAccel{
+  SYMS_FileFormat format;
+  SYMS_MsfAccel *msf;
+  SYMS_PdbDbiAccel dbi;
+  SYMS_PdbNamedStreamArray named;
+  SYMS_PdbStrtblAccel strtbl;
+  SYMS_PdbTpiAccel tpi;
+  SYMS_PdbTpiAccel ipi;
+  SYMS_PdbGsiAccel gsi;
+  SYMS_PdbGsiAccel psi;
+  SYMS_PeGuid auth_guid;
+} SYMS_PdbDbgAccel;
+
+////////////////////////////////
+//~ NOTE(allen): PDB Members
+
+typedef struct SYMS_PdbMemStubNode{
+  struct SYMS_PdbMemStubNode *next;
+  SYMS_String8 name;
+  SYMS_U32 num;
+  SYMS_U32 off;
+  SYMS_U32 off2;
+} SYMS_PdbMemStubNode;
+
+typedef struct SYMS_PdbMemStubList{
+  SYMS_PdbMemStubNode *first;
+  SYMS_PdbMemStubNode *last;
+  SYMS_U64 mem_count;
+} SYMS_PdbMemStubList;
+
+typedef struct SYMS_PdbMemsAccel{
+  SYMS_FileFormat format;
+  SYMS_String8 type_name;
+  SYMS_U64 count;
+  SYMS_PdbMemStubNode **members;
+} SYMS_PdbMemsAccel;
+
+typedef struct SYMS_EnumInfoNode{
+  struct SYMS_EnumInfoNode *next;
+  SYMS_EnumInfo enum_info;
+} SYMS_EnumInfoNode;
+
+////////////////////////////////
+//~ NOTE(allen): PDB Lines
+
+typedef struct SYMS_PdbFileNode{
+  struct SYMS_PdbFileNode *next;
+  SYMS_U64 count;
+  SYMS_FileID file_ids[6];
+} SYMS_PdbFileNode;
+
+typedef struct SYMS_PdbLineSequence{
+  struct SYMS_PdbLineSequence *next;
+  SYMS_Line *lines;
+  SYMS_U64 line_count;
+} SYMS_PdbLineSequence;
+
+typedef struct SYMS_PdbLineTableLoose{
+  SYMS_PdbFileNode *first_file_node;
+  SYMS_PdbFileNode *last_file_node;
+  SYMS_U64 file_count;
+  
+  SYMS_PdbLineSequence *first_seq;
+  SYMS_PdbLineSequence *last_seq;
+  SYMS_U64 seq_count;
+  SYMS_U64 line_count;
+} SYMS_PdbLineTableLoose;
+
+
+////////////////////////////////
+//~ NOTE(allen): PDB Map Accel Type
+
+typedef struct SYMS_PdbMapAccel{
+  SYMS_FileFormat format;
+  SYMS_UnitID uid;
+} SYMS_PdbMapAccel;
+
+////////////////////////////////
+//~ NOTE(allen): PDB TPI Functions
+
+//- tpi accel helpers
+SYMS_API SYMS_PdbTpiOffRange syms_pdb_tpi__hint_from_index(SYMS_PdbTpiAccel *tpi, SYMS_CvTypeIndex ti);
+SYMS_API void                syms_pdb_tpi__fill_off_range(SYMS_String8 data, SYMS_MsfAccel *msf,
+                                                          SYMS_PdbTpiAccel *tpi, SYMS_PdbTpiOffRange *fill);
+
+//- tpi accel
+SYMS_API SYMS_PdbTpiAccel syms_pdb_tpi_accel_from_sn(SYMS_Arena *arena, SYMS_String8 data, SYMS_MsfAccel *msf,
+                                                     SYMS_MsfStreamNumber sn);
+SYMS_API SYMS_U32         syms_pdb_tpi_off_from_ti(SYMS_String8 data, SYMS_MsfAccel *msf, SYMS_PdbTpiAccel *tpi,
+                                                   SYMS_CvTypeIndex ti);
+SYMS_API SYMS_U32         syms_pdb_tpi_base_off(SYMS_PdbTpiAccel *tpi);
+SYMS_API SYMS_MsfRange    syms_pdb_tpi_range(SYMS_MsfAccel *msf, SYMS_PdbTpiAccel *tpi);
+
+//- tpi name lookups
+SYMS_API SYMS_USIDList syms_pdb_types_from_name(SYMS_Arena *arena, SYMS_String8 data, SYMS_PdbDbgAccel *dbg,
+                                                SYMS_PdbUnitAccel *unit, SYMS_String8 name);
+
+
+////////////////////////////////
+//~ NOTE(allen): PDB GSI Functions
+
+SYMS_API SYMS_PdbGsiAccel  syms_pdb_gsi_accel_from_range(SYMS_Arena *arena, SYMS_String8 data, SYMS_MsfAccel *msf,
+                                                         SYMS_MsfRange range);
+SYMS_API SYMS_MsfRange     syms_pdb_gsi_part_from_psi_range(SYMS_MsfRange psi_range);
+
+SYMS_API SYMS_USIDList     syms_pdb_symbols_from_name(SYMS_Arena *arena, SYMS_String8 data, SYMS_MsfAccel *msf,
+                                                      SYMS_PdbGsiAccel *gsi, SYMS_PdbUnitAccel *unit,
+                                                      SYMS_String8 name);
+
+
+////////////////////////////////
+//~ NOTE(allen): PDB Accelerator Setup
+
+// pdb specific parsing
+SYMS_API SYMS_PdbDbiAccel syms_pdb_dbi_from_msf(SYMS_String8 data, SYMS_MsfAccel *msf);
+SYMS_API SYMS_MsfRange    syms_pdb_dbi_sub_range(SYMS_PdbDbiAccel *dbi, SYMS_PdbDbiRange n);
+SYMS_API SYMS_MsfRange    syms_pdb_dbi_stream(SYMS_MsfAccel *msf, SYMS_PdbDbiAccel *dbi, SYMS_PdbDbiStream n);
+
+SYMS_API SYMS_PdbInfoTable        syms_pdb_parse_info(SYMS_Arena *arena, SYMS_String8 data, SYMS_MsfAccel *msf);
+SYMS_API SYMS_PdbNamedStreamArray syms_pdb_named_stream_array(SYMS_PdbInfoTable *table);
+SYMS_API SYMS_PdbStrtblAccel      syms_pdb_dbi_parse_strtbl(SYMS_Arena *arena, SYMS_String8 data,
+                                                            SYMS_MsfAccel *msf, SYMS_MsfStreamNumber sn);
+
+// pdb specific api
+SYMS_API SYMS_MsfAccel* syms_pdb_msf_accel_from_dbg(SYMS_PdbDbgAccel *dbg);
+
+SYMS_API SYMS_String8 syms_pdb_strtbl_string_from_off(SYMS_Arena *arena, SYMS_String8 data, SYMS_PdbDbgAccel *dbg,
+                                                      SYMS_U32 off);
+SYMS_API SYMS_String8 syms_pdb_strtbl_string_from_index(SYMS_Arena *arena, SYMS_String8 data,
+                                                        SYMS_PdbDbgAccel *dbg, SYMS_PdbStringIndex n);
+
+// main api
+SYMS_API SYMS_PdbFileAccel* syms_pdb_file_accel_from_data(SYMS_Arena *arena, SYMS_String8 data);
+SYMS_API SYMS_PdbDbgAccel*  syms_pdb_dbg_accel_from_file(SYMS_Arena *arena, SYMS_String8 data,
+                                                         SYMS_PdbFileAccel *file);
+
+
+////////////////////////////////
+//~ NOTE(allen): PDB Architecture
+
+// main api
+SYMS_API SYMS_Arch syms_pdb_arch_from_dbg(SYMS_PdbDbgAccel *dbg);
+
+
+////////////////////////////////
+//~ NOTE(allen): PDB Match Keys
+
+SYMS_API SYMS_ExtMatchKey syms_pdb_ext_match_key_from_dbg(SYMS_String8 data, SYMS_PdbDbgAccel *dbg);
+
+
+////////////////////////////////
+//~ NOTE(allen): PDB Sections
+
+// pdb specific
+SYMS_API SYMS_CoffSection syms_pdb_coff_section(SYMS_String8 data, SYMS_PdbDbgAccel *dbg, SYMS_U64 n);
+
+// main api
+SYMS_API SYMS_SecInfoArray syms_pdb_sec_info_array_from_dbg(SYMS_Arena *arena, SYMS_String8 data,
+                                                            SYMS_PdbDbgAccel *dbg);
+
+////////////////////////////////
+//~ NOTE(allen): PDB Compilation Units
+
+// pdb specific
+SYMS_API SYMS_PdbCompUnit* syms_pdb_comp_unit_from_id(SYMS_PdbUnitSetAccel *unit_set, SYMS_UnitID id);
+SYMS_API SYMS_MsfRange     syms_pdb_msf_range_from_comp_unit(SYMS_PdbCompUnit *unit, SYMS_PdbCompUnitRange n);
+
+// main api
+SYMS_API SYMS_PdbUnitSetAccel*syms_pdb_unit_set_accel_from_dbg(SYMS_Arena *arena, SYMS_String8 data,
+                                                               SYMS_PdbDbgAccel *dbg);
+SYMS_API SYMS_U64             syms_pdb_unit_count_from_set(SYMS_PdbUnitSetAccel *unit_set);
+SYMS_API SYMS_UnitInfo        syms_pdb_unit_info_from_uid(SYMS_PdbUnitSetAccel *unit_set, SYMS_UnitID uid);
+SYMS_API SYMS_UnitNames       syms_pdb_unit_names_from_uid(SYMS_Arena *arena, SYMS_PdbUnitSetAccel *unit_set,
+                                                           SYMS_UnitID uid);
+
+SYMS_API SYMS_UnitRangeArray  syms_pdb_unit_ranges_from_set(SYMS_Arena *arena,
+                                                            SYMS_String8 data, SYMS_PdbDbgAccel *dbg,
+                                                            SYMS_PdbUnitSetAccel *unit_set);
+
+SYMS_API SYMS_UnitFeatures    syms_pdb_unit_features_from_number(SYMS_PdbUnitSetAccel *unit_set, SYMS_U64 n);
+
+////////////////////////////////
+//~ NOTE(allen): PDB Symbol Parsing
+
+// cv parsing helpers
+SYMS_API void               syms_pdb_leaf_debug_helper(SYMS_CvLeaf kind);
+SYMS_API void               syms_pdb_sym_debug_helper(SYMS_CvSymKind kind);
+
+SYMS_API SYMS_CvElement     syms_cv_element(SYMS_String8 data, SYMS_MsfAccel *msf, SYMS_MsfRange range,
+                                            SYMS_U32 off);
+SYMS_API SYMS_U32           syms_cv_read_numeric(SYMS_String8 data, SYMS_MsfAccel *msf, SYMS_MsfRange range,
+                                                 SYMS_U32 off, SYMS_CvNumeric *out);
+SYMS_API SYMS_U32           syms_cv_u32_from_numeric(SYMS_CvNumeric num);
+
+SYMS_API SYMS_ConstInfo     syms_pdb_const_info_from_cv_numeric(SYMS_CvNumeric num);
+
+SYMS_API SYMS_PdbStubRef*   syms_pdb_alloc_ref(SYMS_Arena *arena, SYMS_PdbStubRef **free_list);
+
+// cv parse
+SYMS_API SYMS_PdbUnitAccel*syms_pdb_sym_accel_from_range(SYMS_Arena *arena, SYMS_String8 data, SYMS_MsfAccel *msf,
+                                                         SYMS_MsfRange range, SYMS_UnitID uid);
+
+SYMS_API SYMS_SymbolKind   syms_pdb_sym_symbol_kind_from_id(SYMS_String8 data, SYMS_PdbDbgAccel *dbg,
+                                                            SYMS_PdbUnitAccel *unit, SYMS_SymbolID id);
+SYMS_API SYMS_String8      syms_pdb_sym_symbol_name_from_id(SYMS_Arena *arena,
+                                                            SYMS_String8 data, SYMS_PdbDbgAccel *dbg,
+                                                            SYMS_PdbUnitAccel *unit, SYMS_SymbolID id);
+
+SYMS_API SYMS_PdbUnitAccel*  syms_pdb_pub_sym_accel_from_dbg(SYMS_Arena *arena, SYMS_String8 data,
+                                                             SYMS_PdbDbgAccel *dbg);
+
+SYMS_API SYMS_PdbUnitAccel*  syms_pdb_leaf_accel_from_dbg(SYMS_Arena *arena, SYMS_String8 data,
+                                                          SYMS_PdbDbgAccel *dbg);
+SYMS_API SYMS_PdbLeafResolve syms_pdb_leaf_resolve_from_id(SYMS_String8 data, SYMS_PdbDbgAccel *dbg,
+                                                           SYMS_PdbUnitAccel *unit, SYMS_SymbolID id);
+
+SYMS_API SYMS_SymbolKind     syms_pdb_leaf_symbol_kind_from_sid(SYMS_String8 data, SYMS_PdbDbgAccel *dbg,
+                                                                SYMS_PdbUnitAccel *unit, SYMS_SymbolID id);
+SYMS_API SYMS_String8        syms_pdb_leaf_symbol_name_from_sid(SYMS_Arena *arena,
+                                                                SYMS_String8 data, SYMS_PdbDbgAccel *dbg,
+                                                                SYMS_PdbUnitAccel *unit, SYMS_SymbolID id);
+
+SYMS_API SYMS_TypeInfo       syms_pdb_leaf_type_info_from_id(SYMS_String8 data, SYMS_PdbDbgAccel *dbg,
+                                                             SYMS_PdbUnitAccel *unit, SYMS_SymbolID id);
+SYMS_API SYMS_ConstInfo      syms_pdb_leaf_const_info_from_id(SYMS_String8 data, SYMS_PdbDbgAccel *dbg,
+                                                              SYMS_PdbUnitAccel *unit, SYMS_SymbolID id);
+
+// pdb unit helpers
+SYMS_API SYMS_PdbStub* syms_pdb_stub_from_unit_off(SYMS_PdbUnitAccel *unit, SYMS_U32 off);
+SYMS_API SYMS_PdbStub* syms_pdb_stub_from_unit_index(SYMS_PdbUnitAccel *unit, SYMS_U32 index);
+
+// main api
+SYMS_API SYMS_PdbUnitAccel*syms_pdb_unit_accel_from_id(SYMS_Arena *arena, SYMS_String8 data, SYMS_PdbDbgAccel *dbg,
+                                                       SYMS_PdbUnitSetAccel *unit_set, SYMS_UnitID id);
+
+SYMS_API SYMS_UnitID       syms_pdb_uid_from_accel(SYMS_PdbUnitAccel *unit);
+
+SYMS_API SYMS_SymbolIDArray syms_pdb_proc_sid_array_from_unit(SYMS_Arena *arena, SYMS_PdbUnitAccel *unit);
+SYMS_API SYMS_SymbolIDArray syms_pdb_var_sid_array_from_unit(SYMS_Arena *arena, SYMS_PdbUnitAccel *unit);
+SYMS_API SYMS_SymbolIDArray syms_pdb_type_sid_array_from_unit(SYMS_Arena *arena, SYMS_PdbUnitAccel *unit);
+
+SYMS_API SYMS_U64          syms_pdb_symbol_count_from_unit(SYMS_PdbUnitAccel *unit);
+
+SYMS_API SYMS_SymbolKind   syms_pdb_symbol_kind_from_sid(SYMS_String8 data, SYMS_PdbDbgAccel *dbg,
+                                                         SYMS_PdbUnitAccel *unit, SYMS_SymbolID id);
+SYMS_API SYMS_String8      syms_pdb_symbol_name_from_sid(SYMS_Arena *arena,SYMS_String8 data,SYMS_PdbDbgAccel *dbg,
+                                                         SYMS_PdbUnitAccel *unit, SYMS_SymbolID id);
+
+
+SYMS_API SYMS_TypeInfo     syms_pdb_type_info_from_id(SYMS_String8 data, SYMS_PdbDbgAccel *dbg,
+                                                      SYMS_PdbUnitAccel *unit, SYMS_SymbolID id);
+SYMS_API SYMS_ConstInfo    syms_pdb_const_info_from_id(SYMS_String8 data, SYMS_PdbDbgAccel *dbg,
+                                                       SYMS_PdbUnitAccel *unit, SYMS_SymbolID id);
+
+////////////////////////////////
+//~ NOTE(allen): PDB Variable Info
+
+// cv parse
+SYMS_API SYMS_USID   syms_pdb_sym_type_from_var_id(SYMS_String8 data, SYMS_PdbDbgAccel *dbg,
+                                                   SYMS_PdbUnitAccel *unit, SYMS_SymbolID id);
+
+SYMS_API SYMS_U64    syms_pdb_sym_voff_from_var_sid(SYMS_String8 data, SYMS_PdbDbgAccel *dbg,
+                                                    SYMS_PdbUnitAccel *unit, SYMS_SymbolID sid);
+
+// main api
+SYMS_API SYMS_USID   syms_pdb_type_from_var_id(SYMS_String8 data, SYMS_PdbDbgAccel *dbg, SYMS_PdbUnitAccel *unit,
+                                               SYMS_SymbolID id);
+
+SYMS_API SYMS_U64    syms_pdb_voff_from_var_sid(SYMS_String8 data, SYMS_PdbDbgAccel *dbg, SYMS_PdbUnitAccel *unit,
+                                                SYMS_SymbolID sid);
+
+
+////////////////////////////////
+//~ NOTE(allen): PDB Member Info
+
+// cv parse
+SYMS_API void syms_pdb__field_list_parse(SYMS_Arena *arena, SYMS_String8 data, SYMS_MsfAccel *msf,
+                                         SYMS_PdbUnitAccel *unit, SYMS_U32 index, SYMS_PdbMemStubList *out);
+
+// main api
+SYMS_API SYMS_PdbMemsAccel*syms_pdb_mems_accel_from_sid(SYMS_Arena *arena, SYMS_String8 data,
+                                                        SYMS_PdbDbgAccel *dbg, SYMS_PdbUnitAccel *unit,
+                                                        SYMS_SymbolID id);
+
+SYMS_API SYMS_U64      syms_pdb_mem_count_from_mems(SYMS_PdbMemsAccel *mems);
+SYMS_API SYMS_MemInfo  syms_pdb_mem_info_from_number(SYMS_Arena *arena, SYMS_String8 data, SYMS_PdbDbgAccel *dbg,
+                                                     SYMS_PdbUnitAccel *unit, SYMS_PdbMemsAccel *mems, SYMS_U64 n);
+
+SYMS_API SYMS_USID     syms_pdb_type_from_mem_number(SYMS_String8 data, SYMS_PdbDbgAccel *dbg,
+                                                     SYMS_PdbUnitAccel *unit, SYMS_PdbMemsAccel *mems,
+                                                     SYMS_U64 n);
+
+SYMS_API SYMS_SigInfo  syms_pdb_sig_info_from_mem_number(SYMS_Arena *arena, SYMS_String8 data,
+                                                         SYMS_PdbDbgAccel *dbg, SYMS_PdbUnitAccel *unit,
+                                                         SYMS_PdbMemsAccel *mems, SYMS_U64 n);
+
+SYMS_API SYMS_EnumInfoArray syms_pdb_enum_info_array_from_sid(SYMS_Arena *arena, SYMS_String8 data,
+                                                              SYMS_PdbDbgAccel *dbg, SYMS_PdbUnitAccel *unit,
+                                                              SYMS_SymbolID sid);
+
+////////////////////////////////
+//~ NOTE(allen): PDB Procedure Info
+
+// main api
+
+SYMS_API SYMS_UnitIDAndSig syms_pdb_proc_sig_handle_from_id(SYMS_String8 data, SYMS_PdbDbgAccel *dbg,
+                                                            SYMS_PdbUnitAccel *unit, SYMS_SymbolID id);
+SYMS_API SYMS_SigInfo      syms_pdb_sig_info_from_handle(SYMS_Arena *arena, SYMS_String8 data,
+                                                         SYMS_PdbDbgAccel *dbg, SYMS_PdbUnitAccel *unit,
+                                                         SYMS_SigHandle handle);
+
+SYMS_API SYMS_U64RangeArray syms_pdb_proc_vranges_from_sid(SYMS_Arena *arena, SYMS_String8 data,
+                                                           SYMS_PdbDbgAccel *dbg, SYMS_PdbUnitAccel *unit,
+                                                           SYMS_SymbolID sid);
+SYMS_API SYMS_SymbolIDArray syms_pdb_scope_children_from_sid(SYMS_Arena *arena, SYMS_String8 data,
+                                                             SYMS_PdbDbgAccel *dbg, SYMS_PdbUnitAccel *unit,
+                                                             SYMS_SymbolID id);
+
+////////////////////////////////
+//~ NOTE(allen): PDB Signature Info
+
+// pdb specific helper
+SYMS_API SYMS_SigInfo syms_pdb_sig_info_from_sig_index(SYMS_Arena *arena, SYMS_String8 data, SYMS_PdbDbgAccel *dbg,
+                                                       SYMS_PdbUnitAccel *unit, SYMS_CvTypeIndex index);
+
+// main api
+SYMS_API SYMS_SigInfo syms_pdb_sig_info_from_id(SYMS_Arena *arena, SYMS_String8 data, SYMS_PdbDbgAccel *dbg,
+                                                SYMS_PdbUnitAccel *unit, SYMS_SymbolID id);
+
+////////////////////////////////
+//~ NOTE(allen): PDB Line Info
+
+// pdb specific
+SYMS_API SYMS_PdbC13SubSection* syms_pdb_c13_sub_sections(SYMS_Arena *arena, SYMS_String8 data, 
+                                                          SYMS_MsfAccel *msf, SYMS_MsfRange range);
+
+SYMS_API void       syms_pdb_loose_push_file_id(SYMS_Arena *arena, SYMS_PdbLineTableLoose *loose, SYMS_FileID id);
+SYMS_API SYMS_Line* syms_pdb_loose_push_sequence(SYMS_Arena *arena, SYMS_PdbLineTableLoose *loose,
+                                                 SYMS_U64 line_count);
+
+// main api
+
+SYMS_API SYMS_String8 syms_pdb_file_name_from_id(SYMS_Arena *arena, SYMS_String8 data, SYMS_PdbDbgAccel *dbg,
+                                                 SYMS_PdbUnitSetAccel *unit_set, SYMS_UnitID uid,
+                                                 SYMS_FileID id);
+
+SYMS_API SYMS_LineParseOut syms_pdb_line_parse_from_uid(SYMS_Arena *arena, SYMS_String8 data,
+                                                        SYMS_PdbDbgAccel *dbg, SYMS_PdbUnitSetAccel *set,
+                                                        SYMS_UnitID uid);
+
+////////////////////////////////
+//~ NOTE(allen): PDB Name Maps
+
+// main api
+SYMS_API SYMS_PdbMapAccel* syms_pdb_type_map_from_dbg(SYMS_Arena *arena, SYMS_String8 data, SYMS_PdbDbgAccel *dbg);
+SYMS_API SYMS_PdbMapAccel* syms_pdb_image_symbol_map_from_dbg(SYMS_Arena *arena, SYMS_String8 data,
+                                                              SYMS_PdbDbgAccel *dbg);
+SYMS_API SYMS_UnitID       syms_pdb_partner_uid_from_map(SYMS_PdbMapAccel *map);
+SYMS_API SYMS_USIDList     syms_pdb_usid_list_from_string(SYMS_Arena *arena, SYMS_String8 data,
+                                                          SYMS_PdbDbgAccel *dbg, SYMS_PdbUnitAccel *unit,
+                                                          SYMS_PdbMapAccel *map, SYMS_String8 string);
+
+////////////////////////////////
+//~ NOTE(allen): PDB CV -> Syms Enums and Flags
+
+SYMS_API SYMS_TypeKind       syms_pdb_type_kind_from_cv_pointer_mode(SYMS_CvPointerMode mode);
+SYMS_API SYMS_TypeModifiers  syms_pdb_modifier_from_cv_pointer_attribs(SYMS_CvPointerAttribs attribs);
+SYMS_API SYMS_TypeModifiers  syms_pdb_modifier_from_cv_modifier_flags(SYMS_CvModifierFlags flags);
+SYMS_API SYMS_CallConvention syms_pdb_call_convention_from_cv_call_kind(SYMS_CvCallKind kind);
+//SYMS_API SYMS_ProcedureFlags syms_pdb_procedure_flags_from_cv_procedure_flags(SYMS_CvProcFlags pdb_flags);
+//SYMS_API SYMS_VarFlags       syms_var_flags_from_cv_local_flags(SYMS_CvLocalFlags pdb_flags);
+
+#endif //SYMS_PDB_PARSER_H
diff --git a/Engine/Source/Runtime/SymsLib/syms/core/pe/syms_pe_parser.c b/Engine/Source/Runtime/SymsLib/syms/core/pe/syms_pe_parser.c
new file mode 100644
--- /dev/null
+++ b/Engine/Source/Runtime/SymsLib/syms/core/pe/syms_pe_parser.c
@@ -0,0 +1,405 @@
+// Copyright Epic Games, Inc. All Rights Reserved.
+#ifndef SYMS_PE_PARSER_C
+#define SYMS_PE_PARSER_C
+
+////////////////////////////////
+// NOTE(allen): Functions
+
+SYMS_API SYMS_PeFileAccel*
+syms_pe_file_accel_from_data(SYMS_Arena *arena, SYMS_String8 data){
+  // based range setup
+  void *base = data.str;
+  SYMS_U64Range range = {0, data.size};
+  
+  // read dos header
+  SYMS_DosHeader dos_header = {0};
+  syms_based_range_read_struct(base, range, 0, &dos_header);
+  
+  // fill result
+  SYMS_PeFileAccel *result = (SYMS_PeFileAccel*)&syms_format_nil;
+  if (dos_header.magic == SYMS_DOS_MAGIC){
+    result = syms_push_array(arena, SYMS_PeFileAccel, 1);
+    result->format = SYMS_FileFormat_PE;
+    result->coff_off = dos_header.coff_file_offset;
+  }
+  
+  return(result);
+}
+
+SYMS_API SYMS_PeBinAccel*
+syms_pe_bin_accel_from_file(SYMS_Arena *arena, SYMS_String8 data, SYMS_PeFileAccel *accel){
+  // based range setup
+  void *base = data.str;
+  SYMS_U64Range range = {0, data.size};
+  
+  // read pe magic
+  SYMS_U32 coff_off = accel->coff_off;
+  SYMS_U32 pe_magic = 0;
+  syms_based_range_read_struct(base, range, coff_off, &pe_magic);
+  
+  // check pe magic
+  SYMS_PeBinAccel *result = (SYMS_PeBinAccel*)&syms_format_nil;
+  if (pe_magic == SYMS_PE_MAGIC){
+    
+    // read coff header
+    SYMS_U32 coff_header_off = coff_off + sizeof(pe_magic);
+    SYMS_CoffHeader coff_header = {0};
+    syms_based_range_read_struct(base, range, coff_header_off, &coff_header);
+    
+    // range of optional extension header ("optional" for short)
+    SYMS_U32 optional_size = coff_header.size_of_optional_header;
+    SYMS_U64 after_coff_header_off = coff_header_off + sizeof(coff_header);
+    SYMS_U64 after_optional_header_off = after_coff_header_off + optional_size;
+    SYMS_U64Range optional_range = {0};
+    optional_range.min = SYMS_ClampTop(after_coff_header_off, data.size);
+    optional_range.max = SYMS_ClampTop(after_optional_header_off, data.size);
+    
+    // get sections
+    SYMS_U64 sec_array_off = optional_range.max;
+    SYMS_U64 sec_array_raw_opl = sec_array_off + coff_header.section_count*sizeof(SYMS_CoffSection);
+    SYMS_U64 sec_array_opl = SYMS_ClampTop(sec_array_raw_opl, data.size);
+    SYMS_U64 clamped_sec_count = (sec_array_opl - sec_array_off)/sizeof(SYMS_CoffSection);
+    SYMS_CoffSection *sections = (SYMS_CoffSection*)((SYMS_U8*)base + sec_array_off);
+    
+    // get data directory from optional
+    SYMS_U16 data_dir_offset = 0;
+    SYMS_U32 data_dir_count = 0;
+    if (optional_size > 0){
+      // read magic number
+      SYMS_U16 optional_magic = 0;
+      syms_based_range_read_struct(base, optional_range, 0, &optional_magic);
+      
+      // read optional
+      SYMS_U32 reported_data_dir_offset = 0;
+      SYMS_U32 reported_data_dir_count = 0;
+      switch (optional_magic){
+        case SYMS_PE32_MAGIC:
+        {
+          SYMS_PeOptionalPe32 pe_optional = {0};
+          syms_based_range_read_struct(base, optional_range, 0, &pe_optional);
+          reported_data_dir_offset = sizeof(pe_optional);
+          reported_data_dir_count = pe_optional.data_dir_count;
+        }break;
+        case SYMS_PE32PLUS_MAGIC:
+        {
+          SYMS_PeOptionalPe32Plus pe_optional = {0};
+          syms_based_range_read_struct(base, optional_range, 0, &pe_optional);
+          reported_data_dir_offset = sizeof(pe_optional);
+          reported_data_dir_count = pe_optional.data_dir_count;
+        }break;
+      }
+      
+      // fill outputs from optional
+      SYMS_U32 data_dir_max = (optional_size - data_dir_offset)/sizeof(SYMS_PeDataDirectory);
+      data_dir_count = SYMS_ClampTop(reported_data_dir_count, data_dir_max);
+      data_dir_offset = optional_range.min + reported_data_dir_offset;
+    }
+    
+    // read info about debug file
+    SYMS_U32 dbg_time = 0;
+    SYMS_U32 dbg_age = 0;
+    SYMS_PeGuid dbg_guid = {0};
+    SYMS_U64 dbg_path_off = 0;
+    SYMS_U64 dbg_path_size = 0;
+    
+    if (SYMS_PeDataDirectoryIndex_DEBUG < data_dir_count){
+      SYMS_PeDataDirectory dbg_data_dir = {0};
+      SYMS_U64 dbg_data_dir_offset = data_dir_offset + sizeof(SYMS_PeDataDirectory)*SYMS_PeDataDirectoryIndex_DEBUG;
+      syms_based_range_read_struct(base, range, dbg_data_dir_offset, &dbg_data_dir);
+      
+      // find section that virtually contains the debug data
+      SYMS_CoffSection *dbg_data_sec = 0;
+      SYMS_CoffSection *sec = sections;
+      for (SYMS_U64 i = 0; i < clamped_sec_count; i += 1, sec += 1){
+        SYMS_U64 first = sec->virt_off;
+        SYMS_U64 opl   = first + sec->virt_size;
+        if (first <= dbg_data_dir.virt_off && dbg_data_dir.virt_off < opl){
+          dbg_data_sec = sec;
+          break;
+        }
+      }
+      
+      // read debug directory
+      if (dbg_data_sec != 0){
+        SYMS_U64 dbg_data_offset_from_sec = dbg_data_dir.virt_off - dbg_data_sec->virt_off;
+        SYMS_U64 dbg_data_offset = dbg_data_sec->file_off + dbg_data_offset_from_sec;
+        SYMS_PeDebugDirectory dbg_data = {0};
+        SYMS_U64 dbg_data_size = SYMS_ClampTop(dbg_data_dir.virt_size, sizeof(dbg_data));
+        syms_based_range_read(base, range, dbg_data_offset, dbg_data_size, &dbg_data);
+        
+        // extract external file info from codeview header
+        if (dbg_data.type == SYMS_PeDebugDirectoryType_CODEVIEW){
+          SYMS_U64 cv_offset = dbg_data.file_offset;
+          SYMS_U32 cv_magic = 0;
+          syms_based_range_read_struct(base, range, cv_offset, &cv_magic);
+          switch (cv_magic){
+            default:break;
+            case SYMS_CODEVIEW_PDB20_MAGIC:
+            {
+              SYMS_PeCvHeaderPDB20 cv = {0};
+              syms_based_range_read_struct(base, range, cv_offset, &cv);
+              // TODO(allen): can we extend the ext match key system to use this in some way?
+              dbg_time = cv.time;
+              dbg_age = cv.age;
+              dbg_path_off = cv_offset + sizeof(cv);
+            }break;
+            case SYMS_CODEVIEW_PDB70_MAGIC:
+            {
+              SYMS_PeCvHeaderPDB70 cv = {0};
+              syms_based_range_read_struct(base, range, cv_offset, &cv);
+              dbg_guid = cv.guid;
+              dbg_age = cv.age;
+              dbg_path_off = cv_offset + sizeof(cv);
+            }break;
+          }
+          if (dbg_path_off > 0){
+            SYMS_String8 dbg_path = syms_based_range_read_string(base, range, dbg_path_off);
+            dbg_path_size = dbg_path.size;
+          }
+        }
+      }
+    }
+    
+    // fill result
+    result = syms_push_array(arena, SYMS_PeBinAccel, 1);
+    result->format = SYMS_FileFormat_PE;
+    result->section_array_off = sec_array_off;
+    result->section_count = clamped_sec_count;
+    result->dbg_path_off = dbg_path_off;
+    result->dbg_path_size = dbg_path_size;
+    result->dbg_guid = dbg_guid;
+    result->arch = syms_arch_from_coff_machine_type(coff_header.machine);
+    
+    
+    // TODO(allen): HACK HACK HACK HACK HACK HACK HACK HACK HACK HACK
+    // This only works because we use it after initializing the rest of PeBinAccel.
+    // The good news is we don't really want to bake this in here long term anyways
+    // so the goal is to get this moved out into a separate UnwindAccel parsing stage.
+    
+    // TODO(allen): only applies for "intel" architectures - need a handler for
+    // arm and mips data structures too.
+    
+    // read pdata
+    SYMS_U64 pdata_off = 0;
+    SYMS_U64 pdata_count = 0;
+    
+    {
+      // try to get pdata section
+      SYMS_CoffSection *pdata_sec = 0;
+      SYMS_CoffSection *coff_sec = (SYMS_CoffSection*)(data.str + sec_array_off);
+      for (SYMS_U64 n = 1; n <= clamped_sec_count; n += 1, coff_sec += 1){
+        
+        // extract name
+        SYMS_U8 *name_base = coff_sec->name;
+        SYMS_U8 *name_ptr = name_base;
+        SYMS_U8 *name_opl = name_base + 8;
+        for (;name_ptr < name_opl && *name_ptr != 0; name_ptr += 1);
+        SYMS_String8 sec_name = syms_str8_range(name_base, name_ptr);
+        
+        // check name
+        if (syms_string_match(sec_name, syms_str8_lit(".pdata"), 0)){
+          pdata_sec = coff_sec;
+          break;
+        }
+      }
+      
+      // determine pdata offset & count
+      if (pdata_sec != 0){
+        SYMS_U64 off = pdata_sec->file_off;
+        SYMS_U64 size = pdata_sec->file_size;
+        
+        // save offset to pdata
+        pdata_off = off;
+        
+        // scan range of pdata to eliminate nulls at the end
+        SYMS_PeIntelPdata *first_pdata_intel = (SYMS_PeIntelPdata*)(data.str + off);
+        SYMS_PeIntelPdata *opl_pdata_intel = first_pdata_intel + (size/sizeof(SYMS_PeIntelPdata));
+        if (first_pdata_intel < opl_pdata_intel){
+          SYMS_PeIntelPdata *pdata_intel_ptr = opl_pdata_intel - 1;
+          for (;pdata_intel_ptr >= first_pdata_intel && pdata_intel_ptr->voff_first == 0;
+               pdata_intel_ptr -= 1);
+          pdata_intel_ptr += 1;
+          pdata_count = (SYMS_U64)(pdata_intel_ptr - first_pdata_intel);
+        }
+      }
+    }
+    
+    // fill pdata
+    result->pdata_off = pdata_off;
+    result->pdata_count = pdata_count;
+  }
+  
+  return(result);
+}
+
+SYMS_API SYMS_Arch
+syms_pe_arch_from_bin(SYMS_PeBinAccel *bin){
+  SYMS_Arch result = bin->arch;
+  return(result);
+}
+
+SYMS_API SYMS_ExtFileList
+syms_pe_ext_file_list_from_bin(SYMS_Arena *arena, SYMS_String8 data, SYMS_PeBinAccel *bin){
+  SYMS_ExtFileList result = {0};
+  if (bin->dbg_path_size > 0 && bin->dbg_path_off + bin->dbg_path_size <= data.size){
+    SYMS_String8 name = {0};
+    name.str = data.str + bin->dbg_path_off;
+    name.size = bin->dbg_path_size;
+    SYMS_ExtFileNode *node = syms_push_array_zero(arena, SYMS_ExtFileNode, 1);
+    SYMS_QueuePush(result.first, result.last, node);
+    result.node_count += 1;
+    node->ext_file.file_name = name;
+    syms_memmove(&node->ext_file.match_key, &bin->dbg_guid, sizeof(bin->dbg_guid));
+  }
+  return(result);
+}
+
+SYMS_API SYMS_CoffSection
+syms_pe_coff_section(SYMS_String8 data, SYMS_PeBinAccel *bin, SYMS_U64 n){
+  SYMS_CoffSection result = {0};
+  if (1 <= n && n <= bin->section_count){
+    SYMS_U64 off = bin->section_array_off + (n - 1)*sizeof(result);
+    syms_memmove(&result, data.str + off, sizeof(result));
+  }
+  return(result);
+}
+
+SYMS_API SYMS_SecInfoArray
+syms_pe_sec_info_array_from_bin(SYMS_Arena *arena, SYMS_String8 data, SYMS_PeBinAccel *bin){
+  SYMS_SecInfoArray result = {0};
+  result.count = bin->section_count;
+  result.sec_info = syms_push_array_zero(arena, SYMS_SecInfo, result.count);
+  
+  SYMS_SecInfo *sec_info = result.sec_info;
+  SYMS_CoffSection *coff_sec = (SYMS_CoffSection*)(data.str + bin->section_array_off);
+  
+  for (SYMS_U64 i = 0; i < result.count; i += 1, sec_info += 1, coff_sec += 1){
+    // extract name
+    SYMS_U8 *name = coff_sec->name;
+    SYMS_U8 *name_ptr = name;
+    SYMS_U8 *name_opl = name + 8;
+    for (;name_ptr < name_opl && *name_ptr != 0; name_ptr += 1);
+    
+    // fill sec info
+    sec_info->name = syms_push_string_copy(arena, syms_str8_range(name, name_ptr));
+    sec_info->vrange.min = coff_sec->virt_off;
+    sec_info->vrange.max = coff_sec->virt_off + coff_sec->virt_size;
+    sec_info->frange.min = coff_sec->file_off;
+    sec_info->frange.max = coff_sec->file_off + coff_sec->file_size;
+  }
+  
+  return(result);
+}
+
+////////////////////////////////
+// NOTE(allen): PE Specific Helpers
+
+SYMS_API SYMS_U64
+syms_pe_binary_search_intel_pdata(SYMS_String8 data, SYMS_PeBinAccel *bin, SYMS_U64 voff){
+  SYMS_U64 result = 0;
+  
+  // check if this bin includes a pdata array
+  if (bin->pdata_count > 0){
+    SYMS_PeIntelPdata *pdata_array = (SYMS_PeIntelPdata*)(data.str + bin->pdata_off);
+    if (voff >= pdata_array[0].voff_first){
+      
+      // binary search:
+      //  find max index s.t. pdata_array[index].voff_first <= voff
+      //  we assume (i < j) -> (pdata_array[i].voff_first < pdata_array[j].voff_first)
+      SYMS_U64 index = bin->pdata_count;
+      SYMS_U64 min = 0;
+      SYMS_U64 opl = bin->pdata_count;
+      for (;;){
+        SYMS_U64 mid = (min + opl)/2;
+        SYMS_PeIntelPdata *pdata = pdata_array + mid;
+        if (voff < pdata->voff_first){
+          opl = mid;
+        }
+        else if (pdata->voff_first < voff){
+          min = mid;
+        }
+        else{
+          index = mid;
+          break;
+        }
+        if (min + 1 >= opl){
+          index = min;
+          break;
+        }
+      }
+      
+      // if we are in range fill result
+      {
+        SYMS_PeIntelPdata *pdata = pdata_array + index;
+        if (pdata->voff_first <= voff && voff < pdata->voff_one_past_last){
+          result = bin->pdata_off + index*sizeof(SYMS_PeIntelPdata);
+        }
+      }
+    }
+  }
+  
+  return(result);
+}
+
+SYMS_API SYMS_U64
+syms_pe_sec_number_from_voff(SYMS_String8 data, SYMS_PeBinAccel *bin, SYMS_U64 voff){
+  SYMS_U64 sec_count = bin->section_count;
+  SYMS_CoffSection *sec_array = (SYMS_CoffSection*)((SYMS_U8*)data.str + bin->section_array_off);
+  SYMS_CoffSection *sec_ptr = sec_array;
+  SYMS_U64 result = 0;
+  for (SYMS_U64 i = 1; i <= sec_count; i += 1, sec_ptr += 1){
+    if (sec_ptr->virt_off <= voff && voff < sec_ptr->virt_off + sec_ptr->virt_size){
+      result = i;
+      break;
+    }
+  }
+  return(result);
+}
+
+SYMS_API void*
+syms_pe_ptr_from_sec_number(SYMS_String8 data, SYMS_PeBinAccel *bin, SYMS_U64 n){
+  void *result = 0;
+  SYMS_U64 sec_count = bin->section_count;
+  if (1 <= n && n <= sec_count){
+    SYMS_CoffSection *sec_array = (SYMS_CoffSection*)((SYMS_U8*)data.str + bin->section_array_off);
+    SYMS_CoffSection *sec = sec_array + n - 1;
+    if (sec->file_size > 0){
+      result = data.str + sec->file_off;
+    }
+  }
+  return(result);
+}
+
+SYMS_API void*
+syms_pe_ptr_from_foff(SYMS_String8 data, SYMS_PeBinAccel *bin, SYMS_U64 foff){
+  void *result = 0;
+  if (foff < data.size){
+    result = data.str + foff;
+  }
+  return(result);
+}
+
+SYMS_API void*
+syms_pe_ptr_from_voff(SYMS_String8 data, SYMS_PeBinAccel *bin, SYMS_U64 voff){
+  //- get the section for this voff
+  SYMS_U64 sec_count = bin->section_count;
+  SYMS_CoffSection *sec_array = (SYMS_CoffSection*)((SYMS_U8*)data.str + bin->section_array_off);
+  SYMS_CoffSection *sec_ptr = sec_array;
+  SYMS_CoffSection *sec = 0;
+  for (SYMS_U64 i = 1; i <= sec_count; i += 1, sec_ptr += 1){
+    if (sec_ptr->virt_off <= voff && voff < sec_ptr->virt_off + sec_ptr->virt_size){
+      sec = sec_ptr;
+      break;
+    }
+  }
+  
+  //- adjust to file pointer
+  void *result = 0;
+  if (sec != 0 & sec_ptr->file_size > 0){
+    result = data.str + voff - sec->virt_off + sec->file_off;
+  }
+  
+  return(result);
+}
+
+#endif //SYMS_PE_PARSER_C
diff --git a/Engine/Source/Runtime/SymsLib/syms/core/syms_parser.h b/Engine/Source/Runtime/SymsLib/syms/core/syms_parser.h
new file mode 100644
--- /dev/null
+++ b/Engine/Source/Runtime/SymsLib/syms/core/syms_parser.h
@@ -0,0 +1,237 @@
+// Copyright Epic Games, Inc. All Rights Reserved.
+/* date = June 8th 2021 10:15 am */
+
+#ifndef SYMS_PARSER_H
+#define SYMS_PARSER_H
+
+////////////////////////////////
+//~ NOTE(allen): Abstraction Acclerators
+
+typedef union SYMS_FileAccel{
+  SYMS_FileFormat format;
+  SYMS_ElfFileAccel elf_accel;
+  SYMS_PeFileAccel pe_accel;
+  SYMS_PdbFileAccel pdb_accel;
+  SYMS_MachFileAccel mach_accel;
+} SYMS_FileAccel;
+
+typedef union SYMS_BinAccel{
+  SYMS_FileFormat format;
+  SYMS_PeBinAccel pe_accel;
+  SYMS_ElfBinAccel elf_accel;
+  SYMS_MachBinAccel mach_accel;
+} SYMS_BinAccel;
+
+typedef union SYMS_BinListAccel{
+  SYMS_FileFormat format;
+  SYMS_MachBinListAccel mach_accel;
+} SYMS_BinListAccel;
+
+typedef union SYMS_DbgAccel{
+  SYMS_FileFormat format;
+  SYMS_DwDbgAccel dw_accel;
+  SYMS_PdbDbgAccel pdb_accel;
+} SYMS_DbgAccel;
+
+typedef union SYMS_UnitSetAccel{
+  SYMS_FileFormat format;
+  SYMS_DwUnitSetAccel dw_accel;
+  SYMS_PdbUnitSetAccel pdb_accel;
+} SYMS_UnitSetAccel;
+
+typedef union SYMS_UnitAccel{
+  SYMS_FileFormat format;
+  SYMS_DwUnitAccel dw_accel;
+  SYMS_PdbUnitAccel pdb_accel;
+} SYMS_UnitAccel;
+
+typedef union SYMS_MemsAccel{
+  SYMS_FileFormat format;
+  SYMS_DwMemsAccel dw_accel;
+  SYMS_PdbMemsAccel pdb_accel;
+} SYMS_MemsAccel;
+
+typedef union SYMS_MapAccel{
+  SYMS_FileFormat format;
+  SYMS_DwMapAccel dw_accel;
+  SYMS_PdbMapAccel pdb_accel;
+} SYMS_MapAccel;
+
+////////////////////////////////
+//~ NOTE(allen): Accelerator Bundle Types
+
+typedef struct SYMS_MapAndUnit{
+  SYMS_MapAccel *map;
+  SYMS_UnitAccel *unit;
+} SYMS_MapAndUnit;
+
+typedef struct SYMS_ParseBundle{
+  SYMS_String8 bin_data;
+  SYMS_String8 dbg_data;
+  SYMS_BinAccel *bin;
+  SYMS_DbgAccel *dbg;
+} SYMS_ParseBundle;
+
+SYMS_C_LINKAGE_BEGIN
+
+
+////////////////////////////////
+//~ NOTE(allen): Accel Helpers
+
+#define syms_accel_is_good(a) ((a) != 0 && (a)->format != SYMS_FileFormat_Null)
+
+////////////////////////////////
+//~ NOTE(allen): General File Analysis
+
+SYMS_API SYMS_FileAccel*   syms_file_accel_from_data(SYMS_Arena *arena, SYMS_String8 data);
+SYMS_API SYMS_FileFormat   syms_file_format_from_file(SYMS_FileAccel *file);
+
+////////////////////////////////
+//~ NOTE(allen): Bin File
+
+SYMS_API SYMS_B32          syms_file_is_bin(SYMS_FileAccel *file);
+SYMS_API SYMS_BinAccel*    syms_bin_accel_from_file(SYMS_Arena *arena, SYMS_String8 data, SYMS_FileAccel *file);
+
+// arch
+SYMS_API SYMS_Arch         syms_arch_from_bin(SYMS_BinAccel *bin);
+
+// external info
+SYMS_API SYMS_ExtFileList  syms_ext_file_list_from_bin(SYMS_Arena *arena, SYMS_String8 data, SYMS_BinAccel *bin);
+
+// binary secs
+SYMS_API SYMS_SecInfoArray syms_sec_info_array_from_bin(SYMS_Arena *arena, SYMS_String8 data, SYMS_BinAccel *bin);
+
+// default vbase
+SYMS_API SYMS_U64          syms_default_vbase_from_bin(SYMS_BinAccel *bin);
+
+////////////////////////////////
+//~ NOTE(nick): Bin List
+
+SYMS_API SYMS_B32           syms_file_is_bin_list(SYMS_FileAccel *file);
+SYMS_API SYMS_BinListAccel* syms_bin_list_from_file(SYMS_Arena *arena, SYMS_String8 data, SYMS_FileAccel *file);
+SYMS_API SYMS_BinInfoArray  syms_bin_info_array_from_bin_list(SYMS_Arena *arena, SYMS_BinListAccel *list);
+SYMS_API SYMS_BinAccel*     syms_bin_accel_from_bin_list_number(SYMS_Arena *arena, SYMS_String8 data,
+                                                                SYMS_BinListAccel *list, SYMS_U64 n);
+////////////////////////////////
+//~ NOTE(rjf): Dbg File
+
+SYMS_API SYMS_B32          syms_file_is_dbg(SYMS_FileAccel *file);
+SYMS_API SYMS_DbgAccel*    syms_dbg_accel_from_file(SYMS_Arena *arena, SYMS_String8 data, SYMS_FileAccel *file);
+
+SYMS_API SYMS_B32          syms_bin_is_dbg(SYMS_BinAccel *bin);
+SYMS_API SYMS_DbgAccel*    syms_dbg_accel_from_bin(SYMS_Arena *arena, SYMS_String8 data, SYMS_BinAccel *bin);
+
+// arch
+SYMS_API SYMS_Arch         syms_arch_from_dbg(SYMS_DbgAccel *dbg);
+
+// external info
+SYMS_API SYMS_ExtFileList  syms_ext_file_list_from_dbg(SYMS_Arena *arena, SYMS_String8 data, SYMS_DbgAccel *dbg);
+
+// match key
+SYMS_API SYMS_ExtMatchKey  syms_ext_match_key_from_dbg(SYMS_String8 data, SYMS_DbgAccel *dbg);
+
+// binary secs
+SYMS_API SYMS_SecInfoArray syms_sec_info_array_from_dbg(SYMS_Arena *arena, SYMS_String8 data, SYMS_DbgAccel *dbg);
+
+// default vbase
+SYMS_API SYMS_U64          syms_default_vbase_from_dbg(SYMS_DbgAccel *dbg);
+
+// compilation units
+SYMS_API SYMS_UnitSetAccel*syms_unit_set_accel_from_dbg(SYMS_Arena *arena, SYMS_String8 data, SYMS_DbgAccel *dbg);
+SYMS_API SYMS_U64          syms_unit_count_from_set(SYMS_UnitSetAccel *unit_set);
+SYMS_API SYMS_UnitInfo     syms_unit_info_from_uid(SYMS_UnitSetAccel *unit_set, SYMS_UnitID uid);
+SYMS_API SYMS_UnitNames    syms_unit_names_from_uid(SYMS_Arena *arena, SYMS_UnitSetAccel *unit_set,
+                                                    SYMS_UnitID uid);
+
+SYMS_API SYMS_UnitRangeArray syms_unit_ranges_from_set(SYMS_Arena *arena, SYMS_String8 data, SYMS_DbgAccel *dbg,
+                                                       SYMS_UnitSetAccel *unit_set);
+
+SYMS_API SYMS_UnitID       syms_uid_collated_types_from_set(SYMS_UnitSetAccel *unit_set);
+SYMS_API SYMS_UnitID       syms_uid_collated_public_symbols_from_set(SYMS_UnitSetAccel *unit_set);
+
+// symbol parsing
+SYMS_API SYMS_UnitAccel*    syms_unit_accel_from_uid(SYMS_Arena *arena, SYMS_String8 data, SYMS_DbgAccel *dbg,
+                                                     SYMS_UnitSetAccel *unit_set, SYMS_UnitID uid);
+
+SYMS_API SYMS_UnitID        syms_uid_from_unit(SYMS_UnitAccel *unit);
+
+SYMS_API SYMS_SymbolIDArray syms_proc_sid_array_from_unit(SYMS_Arena *arena, SYMS_UnitAccel *unit);
+SYMS_API SYMS_SymbolIDArray syms_var_sid_array_from_unit(SYMS_Arena *arena, SYMS_UnitAccel *unit);
+SYMS_API SYMS_SymbolIDArray syms_type_sid_array_from_unit(SYMS_Arena *arena, SYMS_UnitAccel *unit);
+
+
+SYMS_API SYMS_SymbolKind    syms_symbol_kind_from_sid(SYMS_String8 data, SYMS_DbgAccel *dbg,
+                                                      SYMS_UnitAccel *unit, SYMS_SymbolID sid);
+SYMS_API SYMS_String8       syms_symbol_name_from_sid(SYMS_Arena *arena, SYMS_String8 data, SYMS_DbgAccel *dbg,
+                                                      SYMS_UnitAccel *unit, SYMS_SymbolID sid);
+
+
+SYMS_API SYMS_TypeInfo      syms_type_info_from_sid(SYMS_String8 data, SYMS_DbgAccel *dbg, SYMS_UnitAccel *unit,
+                                                    SYMS_SymbolID sid);
+
+SYMS_API SYMS_ConstInfo     syms_const_info_from_sid(SYMS_String8 data, SYMS_DbgAccel *dbg, SYMS_UnitAccel *unit,
+                                                     SYMS_SymbolID sid);
+
+// variable info
+SYMS_API SYMS_USID         syms_type_from_var_sid(SYMS_String8 data, SYMS_DbgAccel *dbg, SYMS_UnitAccel *unit,
+                                                  SYMS_SymbolID sid);
+
+SYMS_API SYMS_U64          syms_voff_from_var_sid(SYMS_String8 data, SYMS_DbgAccel *dbg, SYMS_UnitAccel *unit,
+                                                  SYMS_SymbolID sid);
+
+// member info
+SYMS_API SYMS_MemsAccel*   syms_mems_accel_from_sid(SYMS_Arena *arena, SYMS_String8 data, SYMS_DbgAccel *dbg,
+                                                    SYMS_UnitAccel *unit, SYMS_SymbolID sid);
+
+SYMS_API SYMS_U64          syms_mem_count_from_mems(SYMS_MemsAccel *mems);
+SYMS_API SYMS_MemInfo      syms_mem_info_from_number(SYMS_Arena *arena, SYMS_String8 data, SYMS_DbgAccel *dbg,
+                                                     SYMS_UnitAccel *unit, SYMS_MemsAccel *mems, SYMS_U64 n);
+
+SYMS_API SYMS_USID         syms_type_from_mem_number(SYMS_String8 data, SYMS_DbgAccel *dbg, SYMS_UnitAccel *unit,
+                                                     SYMS_MemsAccel *mems, SYMS_U64 n);
+SYMS_API SYMS_SigInfo      syms_sig_info_from_mem_number(SYMS_Arena *arena, SYMS_String8 data, SYMS_DbgAccel *dbg,
+                                                         SYMS_UnitAccel *unit, SYMS_MemsAccel *mems, SYMS_U64 n);
+
+SYMS_API SYMS_USID         syms_symbol_from_mem_number(SYMS_String8 data, SYMS_DbgAccel *dbg, SYMS_UnitAccel *unit, SYMS_MemsAccel *mems, SYMS_U64 n);
+SYMS_API SYMS_USID         syms_containing_type_from_sid(SYMS_String8 data, SYMS_DbgAccel *dbg, SYMS_UnitAccel *unit, SYMS_SymbolID sid);
+// TODO(rjf): linkage name
+
+SYMS_API SYMS_EnumInfoArray syms_enum_info_array_from_sid(SYMS_Arena *arena, SYMS_String8 data, SYMS_DbgAccel *dbg,
+                                                          SYMS_UnitAccel *unit, SYMS_SymbolID sid);
+
+// procedure info
+SYMS_API SYMS_UnitIDAndSig syms_sig_handle_from_proc_sid(SYMS_String8 data, SYMS_DbgAccel *dbg,
+                                                         SYMS_UnitAccel *unit, SYMS_SymbolID proc_sid);
+SYMS_API SYMS_SigInfo      syms_sig_info_from_handle(SYMS_Arena *arena, SYMS_String8 data, SYMS_DbgAccel *dbg,
+                                                     SYMS_UnitAccel *unit, SYMS_SigHandle handle);
+
+SYMS_API SYMS_U64RangeArray syms_proc_vranges_from_sid(SYMS_Arena *arena, SYMS_String8 data, SYMS_DbgAccel *dbg,
+                                                       SYMS_UnitAccel *unit, SYMS_SymbolID sid);
+SYMS_API SYMS_SymbolIDArray syms_scope_children_from_sid(SYMS_Arena *arena, SYMS_String8 data,
+                                                         SYMS_DbgAccel *dbg, SYMS_UnitAccel *unit,
+                                                         SYMS_SymbolID sid);
+
+// signature info
+SYMS_API SYMS_SigInfo      syms_sig_info_from_type_sid(SYMS_Arena *arena, SYMS_String8 data, SYMS_DbgAccel *dbg,
+                                                       SYMS_UnitAccel *unit, SYMS_SymbolID sid);
+
+// line info
+SYMS_API SYMS_String8      syms_file_name_from_id(SYMS_Arena *arena, SYMS_String8 data, SYMS_DbgAccel *dbg,
+                                                  SYMS_UnitSetAccel *unit_set, SYMS_UnitID uid,
+                                                  SYMS_FileID file_id);
+SYMS_API SYMS_String8Array syms_file_table_from_uid(SYMS_Arena *arena, SYMS_String8 data, SYMS_DbgAccel *dbg,
+                                                    SYMS_UnitSetAccel *unit_set, SYMS_UnitID uid);
+SYMS_API SYMS_LineParseOut syms_line_parse_from_uid(SYMS_Arena *arena, SYMS_String8 data, SYMS_DbgAccel *dbg,
+                                                    SYMS_UnitSetAccel *unit_set, SYMS_UnitID uid);
+
+
+// name maps
+SYMS_API SYMS_MapAccel*    syms_type_map_from_dbg(SYMS_Arena *arena, SYMS_String8 data, SYMS_DbgAccel *dbg);
+SYMS_API SYMS_MapAccel*    syms_image_symbol_map_from_dbg(SYMS_Arena *arena,SYMS_String8 data,SYMS_DbgAccel *dbg);
+SYMS_API SYMS_UnitID       syms_partner_uid_from_map(SYMS_MapAccel *map);
+SYMS_API SYMS_USIDList     syms_symbol_list_from_string(SYMS_Arena *arena, SYMS_String8 data, SYMS_DbgAccel *dbg,
+                                                        SYMS_MapAndUnit map_and_unit, SYMS_String8 string);
+
+SYMS_C_LINKAGE_END
+
+#endif //SYMS_PARSER_H
