Update SymsLib to latest version and update SymsResolver for Insights.

Now it can resolve stripped pdb symbols, for example, which comes from Microsoft Symbol Server.
Insights will additionally use _NT_SYMBOL_PATH environment variable to lookup pdb files.

diff --git a/Engine/Source/Developer/TraceServices/Private/Model/SymslibResolver.cpp b/Engine/Source/Developer/TraceServices/Private/Model/SymslibResolver.cpp
--- a/Engine/Source/Developer/TraceServices/Private/Model/SymslibResolver.cpp
+++ b/Engine/Source/Developer/TraceServices/Private/Model/SymslibResolver.cpp
@@ -50,14 +50,11 @@ namespace
 			return Region.IsValid();
 		}
 
-		inline void* GetData() const
+		SYMS_String8 GetData() const
 		{
-			return Region.IsValid() ? (void*)Region->GetMappedPtr() : nullptr;
-		}
-
-		inline uint64 Num() const
-		{
-			return Region.IsValid() ? Region->GetMappedSize() : 0;
+			return Region.IsValid()
+				? syms_str8((SYMS_U8*)Region->GetMappedPtr(), Region->GetMappedSize())
+				: syms_str8(nullptr, 0);
 		}
 	};
 
@@ -81,91 +78,224 @@ namespace
 			}
 		}
 
-		FString Result = FPaths::Combine(SymbolPath, FileName);
-		if (PlatformFile->FileExists(*Result))
+		if (!SymbolPath.IsEmpty())
 		{
-			// File found in symbol path
-			return Result;
+			FString Result = FPaths::Combine(SymbolPath, FileName);
+			if (PlatformFile->FileExists(*Result))
+			{
+				// File found in symbol path
+				return Result;
+			}
 		}
 
-		if (PlatformFile->FileExists(*Path))
+		// File not found
+		return FString();
+	}
+
+	// use _NT_SYMBOL_PATH environment variable format
+	// for more information see: https://docs.microsoft.com/en-us/windows-hardware/drivers/debugger/advanced-symsrv-use
+	// to explicitly download symbol from MS Symbol Server for specific dll file, run:
+	// "C:\Program Files (x86)\Windows Kits\10\Debuggers\x64\symchk.exe" /r C:\Windows\system32\kernel32.dll /s srv*C:\Symbols*https://msdl.microsoft.com/download/symbols
+	static FString FindWindowsSymbolFile(const FString& GuidPath)
+	{
+		FString SymbolPath = FPlatformMisc::GetEnvironmentVariable(TEXT("_NT_SYMBOL_PATH"));
+		if (SymbolPath.IsEmpty())
 		{
-			// File found by absolute path
-			return Path;
+			return FString();
 		}
 
-		// File not found
+		IPlatformFile* PlatformFile = &FPlatformFileManager::Get().GetPlatformFile();
+
+		TArray<FString> Parts;
+		SymbolPath.ParseIntoArray(Parts, TEXT(";"));
+		for (const FString& Part : Parts)
+		{
+			if (Part.StartsWith(TEXT("srv*")))
+			{
+				TArray<FString> Srv;
+				Part.RightChop(4).ParseIntoArray(Srv, TEXT("*"));
+				for (const FString& Path : Srv)
+				{
+					FString FilePath = FPaths::Combine(Path, GuidPath);
+					if (PlatformFile->FileExists(*FilePath))
+					{
+						return FilePath;
+					}
+				}
+			}
+			else if (!Part.Contains(TEXT("*")))
+			{
+				FString FilePath = FPaths::Combine(Part, GuidPath);
+				if (PlatformFile->FileExists(*FilePath))
+				{
+					return FilePath;
+				}
+			}
+		;}
+
 		return FString();
 	}
 
-	static SYMS_String8 SymsLoader(void* User, SYMS_Arena* Arena, SYMS_String8 FileName)
+	static bool LoadBinary(const TCHAR* Path, SYMS_Arena* Arena, SYMS_ParseBundle& Bundle, TArray<FAutoMappedFile>& Files)
 	{
-		IPlatformFile* PlatformFile = &FPlatformFileManager::Get().GetPlatformFile();
+		FString FilePath = Path;
 
-		FString FilePath = ANSI_TO_TCHAR(reinterpret_cast<char*>(FileName.str));
+		bool bFileFound = false;
+		IPlatformFile* PlatformFile = &FPlatformFileManager::Get().GetPlatformFile();
 
-		TArray<FAutoMappedFile>* Files = static_cast<TArray<FAutoMappedFile>*>(User);
-		FAutoMappedFile& File = Files->AddDefaulted_GetRef();
+		// First lookup file in symbol path
+		FString BinaryPath = FindSymbolFile(FilePath);
+		if (!BinaryPath.IsEmpty())
+		{
+			bFileFound = true;
+		}
 
-		if (Files->Num() == 1)
+		if (!bFileFound)
 		{
-			// we're loading main executable file
-			FString SymbolFilePath = FindSymbolFile(FilePath);
-			if (SymbolFilePath.IsEmpty())
+			// If file is absolute path, try that
+			if (!FPathViews::IsRelativePath(*FilePath) && PlatformFile->FileExists(*FilePath))
 			{
-				UE_LOG(LogSymslib, Warning, TEXT("File '%s' not found"), *FilePath);
-				return syms_str8(nullptr, 0);
+				BinaryPath = FilePath;
+				bFileFound = true;
 			}
+		}
 
-			if (!File.Load(*SymbolFilePath))
-			{
-				UE_LOG(LogSymslib, Warning, TEXT("Failed to open '%s'"), *SymbolFilePath);
-			}
+		if (!bFileFound)
+		{
+			UE_LOG(LogSymslib, Warning, TEXT("Binary file '%s' not found"), *FilePath);
+			return false;
+		}
+
+		FAutoMappedFile& File = Files.AddDefaulted_GetRef();
+		if (!File.Load(*BinaryPath))
+		{
+			UE_LOG(LogSymslib, Warning, TEXT("Failed to load '%s' file"), *BinaryPath);
+			return false;
+		}
+
+		SYMS_FileAccel* Accel = syms_file_accel_from_data(Arena, File.GetData());
+		SYMS_BinAccel* BinAccel = syms_bin_accel_from_file(Arena, File.GetData(), Accel);
+		if (!syms_accel_is_good(BinAccel))
+		{
+			UE_LOG(LogSymslib, Warning, TEXT("Cannot parse '%s' binary file"), *BinaryPath);
+			return false;
+		}
+
+		// remember full path where binary file was found, so debug file can be looked up next to it
+		File.FilePath = BinaryPath;
 
-			File.FilePath = SymbolFilePath;
+		Bundle.bin_data = File.GetData();
+		Bundle.bin = BinAccel;
+
+		return true;
+	}
+
+	static bool LoadDebug(SYMS_Arena* Arena, SYMS_ParseBundle& Bundle, TArray<FAutoMappedFile>& Files)
+	{
+		if (syms_bin_is_dbg(Bundle.bin))
+		{
+			// binary has debug info built-in (like dwarf file)
+			Bundle.dbg = syms_dbg_accel_from_bin(Arena, Files[0].GetData(), Bundle.bin);
+			Bundle.dbg_data = Bundle.bin_data;
+			return true;
 		}
-		else
+
+		// we're loading extra file (pdb for exe)
+		SYMS_ExtFileList List = syms_ext_file_list_from_bin(Arena, Files[0].GetData(), Bundle.bin);
+		if (!List.first)
 		{
-			// we're loading secondary file, like .pdb for .exe
+			UE_LOG(LogSymslib, Warning, TEXT("Binary file '%s' built without debug info"), *Files[0].FilePath);
+			return false;
+		}
+		SYMS_ExtFile ExtFile = List.first->ext_file;
 
-			// main .exe file folder, so we can lookup relative paths
-			FString ModuleBasePath = FPaths::GetPath((*Files)[0].FilePath);
+		// debug file path from metadata in executable
+		FString FilePath = ANSI_TO_TCHAR(reinterpret_cast<char*>(ExtFile.file_name.str));
 
-			bool bDebugFilePathFound = false;
+		bool bFileFound = false;
+		IPlatformFile* PlatformFile = &FPlatformFileManager::Get().GetPlatformFile();
 
-			// First try path itself if it is absolute
-			FString DebugFilePath = FilePath;
-			if (!FPathViews::IsRelativePath(*DebugFilePath))
+		// First try file in symbol path
+		FString DebugPath = FindSymbolFile(FilePath);
+		if (!DebugPath.IsEmpty())
+		{
+			bFileFound = true;
+		}
+
+		if (!bFileFound)
+		{
+			// if executable is PE format, try Windows symbol path format with guid in path
+			if (Bundle.bin->format == SYMS_FileFormat_PE)
 			{
-				if (PlatformFile->FileExists(*DebugFilePath))
+				SYMS_PeBinAccel* PeAccel = reinterpret_cast<SYMS_PeBinAccel*>(Bundle.bin);
+
+				SYMS_PeGuid* Guid = &PeAccel->dbg_guid;
+				SYMS_U32 Age = PeAccel->dbg_age;
+
+				FString FileName = FPaths::GetCleanFilename(FilePath);
+				FString GuidPath = FString::Printf(TEXT("%s/%08X%04X%04X%02X%02X%02X%02X%02X%02X%02X%02X%X/%s"),
+					*FileName,
+					Guid->data1, Guid->data2, Guid->data3,
+					Guid->data4[0], Guid->data4[1], Guid->data4[2], Guid->data4[3], Guid->data4[4], Guid->data4[5], Guid->data4[6], Guid->data4[7],
+					Age,
+					*FileName);
+
+				DebugPath = FindWindowsSymbolFile(GuidPath);
+				if (!DebugPath.IsEmpty())
 				{
-					bDebugFilePathFound = true;
+					bFileFound = true;
 				}
 			}
+		}
 
-			// Otherwise try relative path to module file
-			if (!bDebugFilePathFound)
+		if (!bFileFound)
+		{
+			// Try direct path if it is absolute
+			if (!FPathViews::IsRelativePath(*FilePath) && PlatformFile->FileExists(*FilePath))
 			{
-				DebugFilePath = FPaths::Combine(ModuleBasePath, FPaths::GetCleanFilename(DebugFilePath));
-				if (PlatformFile->FileExists(*DebugFilePath))
-				{
-					bDebugFilePathFound = true;
-				}
+				DebugPath = FilePath;
+				bFileFound = true;
 			}
+		}
 
-			// Use path if file was found
-			if (bDebugFilePathFound)
+		if (!bFileFound)
+		{
+			// Try relative path to module binary file that was found in LoadBinary function
+			FString ModuleBasePath = FPaths::GetPath(Files[0].FilePath);
+			DebugPath = FPaths::Combine(ModuleBasePath, FPaths::GetCleanFilename(FilePath));
+			if (PlatformFile->FileExists(*DebugPath))
 			{
-				if (!File.Load(*DebugFilePath))
-				{
-					UE_LOG(LogSymslib, Warning, TEXT("Failed to open '%s'"), *DebugFilePath);
-				}
-
-				File.FilePath = DebugFilePath;
+				bFileFound = true;
 			}
 		}
 
-		return syms_str8(reinterpret_cast<SYMS_U8*>(File.GetData()), File.Num());
+		if (!bFileFound)
+		{
+			UE_LOG(LogSymslib, Warning, TEXT("Debug file '%s' not found"), *FilePath);
+			return false;
+		}
+
+		FAutoMappedFile& File = Files.AddDefaulted_GetRef();
+		if (!File.Load(*DebugPath))
+		{
+			UE_LOG(LogSymslib, Warning, TEXT("Failed to load '%s' file"), *DebugPath);
+			return false;
+		}
+
+		SYMS_FileAccel* Accel = syms_file_accel_from_data(Arena, File.GetData());
+		SYMS_DbgAccel* DbgAccel = syms_dbg_accel_from_file(Arena, File.GetData(), Accel);
+		if (!syms_accel_is_good(DbgAccel))
+		{
+			UE_LOG(LogSymslib, Warning, TEXT("Cannot parse '%s' debug file"), *DebugPath);
+			return false;
+		}
+
+		File.FilePath = DebugPath;
+
+		Bundle.dbg = DbgAccel;
+		Bundle.dbg_data = File.GetData();
+
+		return true;
 	}
 
 	static bool MatchImageId(const TArray<uint8>& ImageId, SYMS_ParseBundle DataParsed)
@@ -508,28 +638,29 @@ void FSymslibResolver::LoadModuleTracked(FModuleEntry* Module)
 
 FSymslibResolver::EModuleStatus FSymslibResolver::LoadModule(FModuleEntry* Module)
 {
-	IPlatformFile* PlatformFile = &FPlatformFileManager::Get().GetPlatformFile();
-
 	// temporary memory used for loading
 	SYMS_Group* Group = syms_group_alloc();
 
-	SYMS_FileInfOptions opts = { 0 };
-	opts.disable_fallback = 1;
-
+	// memory-mapped binary & debug files
 	TArray<FAutoMappedFile> Files;
-	SYMS_FileLoadCtx Context = { &SymsLoader, &Files };
-	SYMS_FileInfResult FileInf = syms_group_infer_from_file(Group->arena, Context, syms_str8_cstring(TCHAR_TO_ANSI(Module->Path)), &opts);
 
-	SYMS_DbgAccel* Dbg = FileInf.data_parsed.dbg;
-	if (Dbg == nullptr || Dbg->format == SYMS_FileFormat_Null)
+	// contents of binary & debug file
+	SYMS_ParseBundle Bundle;
+
+	if (!LoadBinary(Module->Path, Group->arena, Bundle, Files))
+	{
+		syms_group_release(Group);
+		return EModuleStatus::Failed;
+	}
+
+	if (!LoadDebug(Group->arena, Bundle, Files))
 	{
 		syms_group_release(Group);
-		UE_LOG(LogSymslib, Display, TEXT("No debug information loaded for '%s'"), Module->Name);
 		return EModuleStatus::Failed;
 	}
 
 	// check debug data mismatch to captured module
-	if (!Module->ImageId.IsEmpty() && !MatchImageId(Module->ImageId, FileInf.data_parsed))
+	if (!Module->ImageId.IsEmpty() && !MatchImageId(Module->ImageId, Bundle))
 	{
 		syms_group_release(Group);
 		UE_LOG(LogSymslib, Warning, TEXT("Symbols for '%s' does not match traced binary."), Module->Name);
@@ -539,12 +670,8 @@ FSymslibResolver::EModuleStatus FSymslibResolver::LoadModule(FModuleEntry* Modul
 	FSymsInstance* Instance = &Module->Instance;
 
 	// initialize group
-	SYMS_GroupInitParams GroupInit = {};
-	GroupInit.dbg_data = FileInf.data_parsed.dbg_data;
-	GroupInit.dbg = FileInf.data_parsed.dbg;
-
 	syms_set_lane(0);
-	syms_group_init(Group, &GroupInit);
+	syms_group_init(Group, &Bundle);
 
 	// unit storage
 	SYMS_U64 UnitCount = syms_group_unit_count(Group);
@@ -644,7 +771,34 @@ FSymslibResolver::EModuleStatus FSymslibResolver::LoadModule(FModuleEntry* Modul
 		FPlatformTLS::FreeTlsSlot(LaneSlot);
 	}
 
-	Instance->UnitMap = syms_spatial_map_1d_copy(Instance->Arenas[0], syms_group_unit_map(Group));
+	SYMS_Arena* Arena = Instance->Arenas[0];
+
+	// store stripped format symbols
+	{
+		SYMS_StrippedInfoArray StrippedInfo = syms_group_stripped_info(Group);
+		SYMS_SpatialMap1D* StrippedMap = syms_group_stripped_info_map(Group);
+
+		FSymsSymbol* StrippedSymbols = syms_push_array(Arena, FSymsSymbol, StrippedInfo.count);
+		{
+			SYMS_StrippedInfo* InfoPtr = StrippedInfo.info;
+			FSymsSymbol* StrippedSymbolPtr = StrippedSymbols;
+			for (SYMS_U64 Index = 0; Index < StrippedInfo.count; Index++)
+			{
+				SYMS_String8 Name = syms_push_string_copy(Arena, InfoPtr->name);
+				StrippedSymbolPtr->Name = reinterpret_cast<char*>(Name.str);
+				++StrippedSymbolPtr;
+				++InfoPtr;
+			}
+		}
+
+		Instance->StrippedMap = syms_spatial_map_1d_copy(Arena, StrippedMap);
+		Instance->StrippedSymbols = StrippedSymbols;
+		Instance->StrippedCount = StrippedInfo.count;
+
+		SymbolCount += StrippedInfo.count;
+	}
+
+	Instance->UnitMap = syms_spatial_map_1d_copy(Arena, syms_group_unit_map(Group));
 	Instance->DefaultBase = syms_group_default_vbase(Group);
 
 	syms_group_release(Group);
@@ -702,9 +856,10 @@ bool FSymslibResolver::ResolveSymbol(uint64 Address, FResolvedSymbol* Target, FS
 	SYMS_U64 VirtualOffset = Address + Instance->DefaultBase - Module->Base;
 
 	FSymsSymbol* SymsSymbol = nullptr;
-	bool bFileValid = false;
-	SYMS_String8 FileName = {};
-	uint32 FileLine = 0;
+
+	constexpr uint32 MaxStringSize = 1024;
+	const TCHAR* SourceFilePersistent = nullptr;
+	uint32 SourceFileLine = 0;
 
 	SYMS_UnitID UnitID = syms_spatial_map_1d_value_from_point(&Instance->UnitMap, VirtualOffset);
 	if (UnitID > 0 && UnitID <= Instance->Units.Num())
@@ -722,33 +877,44 @@ bool FSymslibResolver::ResolveSymbol(uint64 Address, FResolvedSymbol* Target, FS
 				SYMS_Line Line = syms_line_from_sequence_voff(&Unit->LineTable, SeqNumber, VirtualOffset);
 				if (Line.src_coord.file_id > 0 && Line.src_coord.file_id <= Unit->FileTable.count)
 				{
-					FileName = Unit->FileTable.strings[Line.src_coord.file_id - 1];
-					FileLine = Line.src_coord.line;
-					bFileValid = true;
+					SYMS_String8 FileName = Unit->FileTable.strings[Line.src_coord.file_id - 1];
+
+					ANSICHAR SourceFile[MaxStringSize];
+					FCStringAnsi::Strncpy(SourceFile, reinterpret_cast<char*>(FileName.str), MaxStringSize);
+					SourceFile[MaxStringSize - 1] = 0;
+					SourceFilePersistent = StringAllocator.Store(ANSI_TO_TCHAR(SourceFile));
+
+					SourceFileLine = Line.src_coord.line;
 				}
 			}
 		}
 	}
 
+	if (SymsSymbol == nullptr)
+	{
+		// try lookup into stripped format symbols
+		SYMS_U64 Value = syms_spatial_map_1d_value_from_point(&Instance->StrippedMap, VirtualOffset);
+		if (Value > 0 && Value <= Instance->StrippedCount)
+		{
+			SYMS_U64 Index = Value - 1;
+			SymsSymbol = &Instance->StrippedSymbols[Index];
+
+			// use module name as filename
+			SourceFilePersistent = StringAllocator.Store(Module->Name);
+		}
+	}
+
 	// this includes skipping symbols without name (empty string)
-	if (!SymsSymbol || !bFileValid || SymsSymbol && SymsSymbol->Name[0] == 0)
+	if (!SymsSymbol || !SourceFilePersistent || SymsSymbol && SymsSymbol->Name[0] == 0)
 	{
 		UpdateResolvedSymbol(Target, ESymbolQueryResult::NotFound, Module->Name, GUnknownModuleTextSymsLib, GUnknownModuleTextSymsLib, 0);
 		return false;
 	}
 
-	constexpr uint32 MaxStringSize = 1024;
-
 	ANSICHAR SymbolName[MaxStringSize];
 	FCStringAnsi::Strncpy(SymbolName, SymsSymbol->Name, MaxStringSize);
 	SymbolName[MaxStringSize - 1] = 0;
-
-	ANSICHAR SourceFile[MaxStringSize];
-	FCStringAnsi::Strncpy(SourceFile, reinterpret_cast<char*>(FileName.str), MaxStringSize);
-	SourceFile[MaxStringSize - 1] = 0;
-
 	const TCHAR* SymbolNamePersistent =  StringAllocator.Store(ANSI_TO_TCHAR(SymbolName));
-	const TCHAR* SourceFilePersistent = StringAllocator.Store(ANSI_TO_TCHAR(SourceFile));
 
 	// Store the strings and update the target data
 	UpdateResolvedSymbol(
@@ -757,7 +923,7 @@ bool FSymslibResolver::ResolveSymbol(uint64 Address, FResolvedSymbol* Target, FS
 		Module->Name,
 		SymbolNamePersistent,
 		SourceFilePersistent,
-		FileLine
+		SourceFileLine
 	);
 
 	return true;
diff --git a/Engine/Source/Developer/TraceServices/Private/Model/SymslibResolver.h b/Engine/Source/Developer/TraceServices/Private/Model/SymslibResolver.h
--- a/Engine/Source/Developer/TraceServices/Private/Model/SymslibResolver.h
+++ b/Engine/Source/Developer/TraceServices/Private/Model/SymslibResolver.h
@@ -54,6 +54,9 @@ private:
 		TArray<SYMS_Arena*> Arenas;
 		TArray<FSymsUnit>   Units;
 		SYMS_SpatialMap1D   UnitMap;
+		SYMS_SpatialMap1D   StrippedMap;
+		FSymsSymbol*        StrippedSymbols;
+		uint64              StrippedCount;
 		uint64              DefaultBase;
 	};
 
diff --git a/Engine/Source/Runtime/SymsLib/syms/core/base/syms_base.c b/Engine/Source/Runtime/SymsLib/syms/core/base/syms_base.c
--- a/Engine/Source/Runtime/SymsLib/syms/core/base/syms_base.c
+++ b/Engine/Source/Runtime/SymsLib/syms/core/base/syms_base.c
@@ -20,6 +20,14 @@ syms_make_u64_range(SYMS_U64 min, SYMS_U64 max){
   return(result);
 }
 
+SYMS_API SYMS_U64Range
+syms_make_u64_inrange(SYMS_U64Range range, SYMS_U64 offset, SYMS_U64 size)
+{
+  SYMS_ASSERT(range.min + offset + size <= range.max);
+  SYMS_U64Range result = syms_make_u64_range(range.min + offset, range.min + offset + size);
+  return result;
+}
+
 SYMS_API SYMS_U64
 syms_u64_range_size(SYMS_U64Range range){
   SYMS_U64 result = range.max - range.min;
@@ -378,6 +386,35 @@ syms_string_trunc_symbol_heuristic(SYMS_String8 string){
   return(result);
 }
 
+SYMS_API SYMS_String8List
+syms_string_split(SYMS_Arena *arena, SYMS_String8 input, SYMS_U32 delimiter)
+{
+  SYMS_String8List result;
+  syms_memzero_struct(&result);
+  
+  SYMS_U8 *ptr = input.str;
+  SYMS_U8 *end = input.str + input.size;
+  SYMS_U64 size = 0;
+  for (;;) {
+    SYMS_U32 cp = 0;
+    SYMS_U8 *pn = syms_decode_utf8(ptr, &cp);
+    if (cp == delimiter) {
+      SYMS_String8 split = syms_str8(ptr - size, size);
+      syms_string_list_push(arena, &result, split);
+      size = 0;
+    } else if (pn >= end || cp == 0) {
+      SYMS_ASSERT(pn == end);
+      SYMS_String8 last = syms_str8(pn - size, size);
+      syms_string_list_push(arena, &result, last);
+      break;
+    }
+    size += 1;
+    ptr = pn;
+  }
+  
+  return result;
+}
+
 ////////////////////////////////
 //~ NOTE(allen): String -> Int
 
@@ -621,7 +658,7 @@ syms_get_lane(void){
 }
 
 ////////////////////////////////
-//~ rjf: Based Ranges
+//~ NOTE(rjf): Based Ranges
 
 SYMS_API void*
 syms_based_range_ptr(void *base, SYMS_U64Range range, SYMS_U64 offset){
@@ -716,16 +753,55 @@ syms_based_range_read_string(void *base, SYMS_U64Range range, SYMS_U64 offset){
   return(result);
 }
 
+////////////////////////////////
+//~ NOTE(allen): Memory Views
+
+SYMS_API SYMS_MemoryView
+syms_memory_view_make(SYMS_String8 data, SYMS_U64 base){
+  SYMS_MemoryView result = {0};
+  result.data = data.str;
+  result.addr_first = base;
+  result.addr_opl = base + data.size;
+  return(result);
+}
+
+SYMS_API SYMS_B32
+syms_memory_view_read(SYMS_MemoryView *memview, SYMS_U64 addr, SYMS_U64 size, void *ptr){
+  SYMS_B32 result = syms_false;
+  if (memview->addr_first <= addr &&
+      addr <= addr + size &&
+      addr + size <= memview->addr_opl){
+    result = syms_true;
+    syms_memmove(ptr, (SYMS_U8*)memview->data + addr - memview->addr_first, size);
+  }
+  return(result);
+}
+
+SYMS_API void
+syms_unwind_result_missed_read(SYMS_UnwindResult *unwind_result, SYMS_U64 addr){
+  unwind_result->dead = syms_true;
+  unwind_result->missed_read = syms_true;
+  unwind_result->missed_read_addr = addr;
+}
+
 ////////////////////////////////
 //~ NOTE(nick): Bit manipulations
 
+SYMS_API SYMS_U16
+syms_bswap_u16(SYMS_U16 x)
+{
+  SYMS_U16 result = (((x & 0xFF00) >> 8) |
+                     ((x & 0x00FF) << 8));
+  return result;
+}
+
 SYMS_API SYMS_U32
 syms_bswap_u32(SYMS_U32 x)
 {
-  SYMS_U32 result = ((x & 0xFF000000) >> 24) |
-  ((x & 0x00FF0000) >> 8)  |
-  ((x & 0x0000FF00) << 8)  |
-  ((x & 0x000000FF) << 24);
+  SYMS_U32 result = (((x & 0xFF000000) >> 24) |
+                     ((x & 0x00FF0000) >> 8)  |
+                     ((x & 0x0000FF00) << 8)  |
+                     ((x & 0x000000FF) << 24));
   return result;
 }
 
@@ -733,17 +809,26 @@ SYMS_API SYMS_U64
 syms_bswap_u64(SYMS_U64 x)
 {
   // TODO(nick): naive bswap, replace with something that is faster like an intrinsic
-  SYMS_U64 result = ((x & 0xFF00000000000000ULL) >> 56) |
-  ((x & 0x00FF000000000000ULL) >> 40) |
-  ((x & 0x0000FF0000000000ULL) >> 24) |
-  ((x & 0x000000FF00000000ULL) >> 8)  |
-  ((x & 0x00000000FF000000ULL) << 8)  |
-  ((x & 0x0000000000FF0000ULL) << 24) |
-  ((x & 0x000000000000FF00ULL) << 40) |
-  ((x & 0x00000000000000FFULL) << 56);
+  SYMS_U64 result = (((x & 0xFF00000000000000ULL) >> 56) |
+                     ((x & 0x00FF000000000000ULL) >> 40) |
+                     ((x & 0x0000FF0000000000ULL) >> 24) |
+                     ((x & 0x000000FF00000000ULL) >> 8)  |
+                     ((x & 0x00000000FF000000ULL) << 8)  |
+                     ((x & 0x0000000000FF0000ULL) << 24) |
+                     ((x & 0x000000000000FF00ULL) << 40) |
+                     ((x & 0x00000000000000FFULL) << 56));
   return result;
 }
 
+SYMS_API void
+syms_bswap_bytes(void *p, SYMS_U64 size)
+{
+  for(SYMS_U32 i = 0, k = size - 1; i < size / 2; ++i, --k)
+  {
+    SYMS_Swap(SYMS_U8, ((SYMS_U8*)p)[k], ((SYMS_U8*)p)[i]);
+  }
+}
+
 ////////////////////////////////
 //~ rjf: Generated Code
 
diff --git a/Engine/Source/Runtime/SymsLib/syms/core/base/syms_base.h b/Engine/Source/Runtime/SymsLib/syms/core/base/syms_base.h
--- a/Engine/Source/Runtime/SymsLib/syms/core/base/syms_base.h
+++ b/Engine/Source/Runtime/SymsLib/syms/core/base/syms_base.h
@@ -87,7 +87,6 @@ typedef double   SYMS_F64;
 #define syms_false 0
 #define syms_true 1
 
-// TODO(allen): eliminate
 typedef SYMS_U32 SYMS_RegID;
 
 ////////////////////////////////
@@ -345,6 +344,27 @@ typedef struct SYMS_ArenaTemp{
   SYMS_U64 pos;
 } SYMS_ArenaTemp;
 
+////////////////////////////////
+//~ NOTE(allen): Memory Views
+
+typedef struct SYMS_MemoryView{
+  // NOTE(allen): Upgrade path:
+  //  1. A list of ranges like this one
+  //  2. After building the list put into a binary-searchable format
+  //  3. Equip with the ability to request missing memory in-line.
+  void *data;
+  SYMS_U64 addr_first;
+  SYMS_U64 addr_opl;
+} SYMS_MemoryView;
+
+typedef struct SYMS_UnwindResult{
+  SYMS_B32 dead;
+  SYMS_B32 missed_read;
+  SYMS_U64 missed_read_addr;
+  SYMS_U64 stack_pointer;
+} SYMS_UnwindResult;
+
+
 ////////////////////////////////
 //~ NOTE(allen): Generated Serial Info
 
@@ -362,6 +382,7 @@ SYMS_API SYMS_String8 syms_version_string(void);
 //~ NOTE(rjf): Basic Type Functions
 
 SYMS_API SYMS_U64Range syms_make_u64_range(SYMS_U64 min, SYMS_U64 max);
+SYMS_API SYMS_U64Range syms_make_u64_inrange(SYMS_U64Range range, SYMS_U64 offset, SYMS_U64 size);
 SYMS_API SYMS_U64 syms_u64_range_size(SYMS_U64Range range);
 
 ////////////////////////////////
@@ -376,8 +397,6 @@ SYMS_API SYMS_U64 syms_hash_u64(SYMS_U64 x);
 ////////////////////////////////
 //~ NOTE(rjf): Serial Information Functions
 
-// TODO(allen): another pass over these ideas
-
 #define syms_serial_type(name) (_syms_serial_type_##name)
 #define syms_string_from_enum_value(enum_type, value) \
 (syms_serial_value_from_enum_value(&syms_serial_type(enum_type), value)->name)
@@ -426,6 +445,8 @@ SYMS_API SYMS_String8 syms_push_string_copy(SYMS_Arena *arena, SYMS_String8 stri
 
 SYMS_API SYMS_String8 syms_string_trunc_symbol_heuristic(SYMS_String8 string);
 
+SYMS_API SYMS_String8List syms_string_split(SYMS_Arena *arena, SYMS_String8 input, SYMS_U32 delimiter);
+
 ////////////////////////////////
 //~ NOTE(allen): String <-> Integer
 
@@ -491,11 +512,24 @@ SYMS_API SYMS_String8 syms_based_range_read_string(void *base, SYMS_U64Range ran
 
 #define syms_based_range_read_struct(b,r,o,p) syms_based_range_read((b), (r), (o), sizeof(*(p)), p)
 
+////////////////////////////////
+//~ NOTE(allen): Memory Views
+
+SYMS_API SYMS_MemoryView syms_memory_view_make(SYMS_String8 data, SYMS_U64 base);
+SYMS_API SYMS_B32        syms_memory_view_read(SYMS_MemoryView *memview, SYMS_U64 addr,
+                                               SYMS_U64 size, void *ptr);
+
+#define syms_memory_view_read_struct(s,a,p) syms_memory_view_read((s),(a),sizeof(*(p)),(p))
+
+SYMS_API void syms_unwind_result_missed_read(SYMS_UnwindResult *unwind_result, SYMS_U64 addr);
+
 ////////////////////////////////
 //~ NOTE(nick): Bit manipulations
 
+SYMS_API SYMS_U16 syms_bswap_u16(SYMS_U16 x);
 SYMS_API SYMS_U32 syms_bswap_u32(SYMS_U32 x);
 SYMS_API SYMS_U64 syms_bswap_u64(SYMS_U64 x);
+SYMS_API void syms_bswap_bytes(void *p, SYMS_U64 size);
 
 ////////////////////////////////
 //~ NOTE(allen): Dev Features
diff --git a/Engine/Source/Runtime/SymsLib/syms/core/base/syms_dev.c b/Engine/Source/Runtime/SymsLib/syms/core/base/syms_dev.c
--- a/Engine/Source/Runtime/SymsLib/syms/core/base/syms_dev.c
+++ b/Engine/Source/Runtime/SymsLib/syms/core/base/syms_dev.c
@@ -80,6 +80,12 @@ syms_log_close__dev(SYMS_U32 prev_state){
   syms_log_state__dev = prev_state;
 }
 
+SYMS_API SYMS_B32
+syms_log_is_enabled__dev(void){
+  SYMS_B32 result = (syms_log_state__dev == 2);
+  return(result);
+}
+
 SYMS_API SYMS_U32
 syms_log_open_annotated__dev(SYMS_LogFeatures features, SYMS_U64 uid){
   SYMS_LogFeatures filter_features = syms_log_filter_features__dev;
@@ -96,12 +102,18 @@ syms_log_open_annotated__dev(SYMS_LogFeatures features, SYMS_U64 uid){
     for (SYMS_U32 i = 0; i < 1; i += 1){
       SYMS_U32 bit = (1 << i);
       if ((features & bit) != 0){
+        SYMS_String8 feature_str = {0};
         switch (bit){
-          case SYMS_LogFeature_LineTable:
-          {
-            syms_string_list_push(scratch.arena, &features_list, syms_str8_lit("line_table"));
-          }break;
+          case SYMS_LogFeature_LineTable:       feature_str = syms_str8_lit("line_table");        break;
+          case SYMS_LogFeature_DwarfUnitRanges: feature_str = syms_str8_lit("dwarf_unit_ranges"); break;
+          case SYMS_LogFeature_DwarfTags:       feature_str = syms_str8_lit("dwarf_tags");        break;
+          case SYMS_LogFeature_DwarfUnwind:     feature_str = syms_str8_lit("dwarf_unwind");      break;
+          case SYMS_LogFeature_DwarfCFILookup:  feature_str = syms_str8_lit("dwarf_cfi_lookup");  break;
+          case SYMS_LogFeature_DwarfCFIDecode:  feature_str = syms_str8_lit("dwarf_cfi_decode");  break;
+          case SYMS_LogFeature_DwarfCFIApply:   feature_str = syms_str8_lit("dwarf_cfi_apply");   break;
+          case SYMS_LogFeature_PEEpilog:        feature_str = syms_str8_lit("pe_epilog");         break;
         }
+        syms_string_list_push(scratch.arena, &features_list, feature_str);
       }
     }
     SYMS_String8 features = syms_str8_lit("misc");
diff --git a/Engine/Source/Runtime/SymsLib/syms/core/base/syms_dev.h b/Engine/Source/Runtime/SymsLib/syms/core/base/syms_dev.h
--- a/Engine/Source/Runtime/SymsLib/syms/core/base/syms_dev.h
+++ b/Engine/Source/Runtime/SymsLib/syms/core/base/syms_dev.h
@@ -15,17 +15,21 @@
 ////////////////////////////////
 // NOTE(allen): Dev Logging Types
 
-// TODO(allen): move to mc
 typedef SYMS_U32 SYMS_LogFeatures;
 enum{
   SYMS_LogFeature_LineTable       = (1 << 0),
   SYMS_LogFeature_DwarfUnitRanges = (1 << 1),
   SYMS_LogFeature_DwarfTags       = (1 << 2),
+  SYMS_LogFeature_DwarfUnwind     = (1 << 3),
+  SYMS_LogFeature_DwarfCFILookup  = (1 << 4),
+  SYMS_LogFeature_DwarfCFIDecode  = (1 << 5),
+  SYMS_LogFeature_DwarfCFIApply   = (1 << 6),
+  SYMS_LogFeature_PEEpilog        = (1 << 7),
   
   
   // Dummy flag shouldn't be associated to anything.
   // This gives us a way to disable all logging.
-  SYMS_LogFeature_Dummy     = (1 << 31),
+  SYMS_LogFeature_Dummy = (1 << 31),
 };
 
 ////////////////////////////////
@@ -119,6 +123,7 @@ SYMS_API void     syms_log_set_filter__dev(SYMS_LogFeatures features, SYMS_U64 u
 
 SYMS_API SYMS_U32 syms_log_open__dev(SYMS_B32 enabled);
 SYMS_API void     syms_log_close__dev(SYMS_U32 prev_state);
+SYMS_API SYMS_B32 syms_log_is_enabled__dev(void);
 SYMS_API SYMS_U32 syms_log_open_annotated__dev(SYMS_LogFeatures features, SYMS_U64 uid);
 
 SYMS_API void syms_logfv__dev(char *fmt, va_list args);
@@ -130,10 +135,12 @@ SYMS_API void syms_logf__dev(char *fmt, ...);
 # define SYMS_LogOpen(ftr,uid,block) SYMS_U32 syms_log_state__##block = syms_log_open_annotated__dev((ftr),(uid))
 # define SYMS_LogClose(block) syms_log_close__dev(syms_log_state__##block)
 # define SYMS_Log(...) syms_logf__dev(__VA_ARGS__)
+# define SYMS_LogIsEnabled() syms_log_is_enabled__dev()
 #else
 # define SYMS_LogOpen(ftr,uid,block) ((void)0)
 # define SYMS_LogClose(block) ((void)0)
 # define SYMS_Log(fmt,...) ((void)0)
+# define SYMS_LogIsEnabled() (0)
 #endif
 
 ////////////////////////////////
diff --git a/Engine/Source/Runtime/SymsLib/syms/core/base/syms_serial.h b/Engine/Source/Runtime/SymsLib/syms/core/base/syms_serial.h
--- a/Engine/Source/Runtime/SymsLib/syms/core/base/syms_serial.h
+++ b/Engine/Source/Runtime/SymsLib/syms/core/base/syms_serial.h
@@ -26,6 +26,12 @@
 # if defined(_SYMS_META_PE_H)
 #  define SYMS_ENABLE_PE_SERIAL_INFO
 # endif
+# if defined(_SYMS_META_MACH_H)
+#  define SYMS_ENABLE_MACH_SERIAL_INFO
+# endif
+# if defined(_SYMS_META_ELF_H)
+#  define SYMS_ENABLE_ELF_SERIAL_INFO
+# endif
 #endif
 
 ////////////////////////////////
diff --git a/Engine/Source/Runtime/SymsLib/syms/core/data_structures/syms_data_structures.c b/Engine/Source/Runtime/SymsLib/syms/core/data_structures/syms_data_structures.c
--- a/Engine/Source/Runtime/SymsLib/syms/core/data_structures/syms_data_structures.c
+++ b/Engine/Source/Runtime/SymsLib/syms/core/data_structures/syms_data_structures.c
@@ -226,10 +226,16 @@ syms_spatial_map_1d_array_sort__rec(SYMS_SpatialMap1DRange *ranges, SYMS_U64 cou
     SYMS_Swap(SYMS_SpatialMap1DRange, ranges[mid], ranges[last]);
     
     // partition
+    SYMS_B32 equal_send_left = 0;
     SYMS_U64 key = ranges[last].range.min;
     SYMS_U64 j = 0;
     for (SYMS_U64 i = 0; i < last; i += 1){
-      if (ranges[i].range.min < key){
+      SYMS_B32 send_left = (ranges[i].range.min < key);
+      if (!send_left && ranges[i].range.min == key){
+        send_left = equal_send_left;
+        equal_send_left = !equal_send_left;
+      }
+      if (send_left){
         if (j != i){
           SYMS_Swap(SYMS_SpatialMap1DRange, ranges[i], ranges[j]);
         }
@@ -799,4 +805,22 @@ syms_string_array_copy(SYMS_Arena *arena, SYMS_StringCons *cons, SYMS_String8Arr
   return(result);
 }
 
+SYMS_API SYMS_StrippedInfoArray
+syms_stripped_info_copy(SYMS_Arena *arena, SYMS_StrippedInfoArray *stripped){
+  SYMS_ProfBegin("syms_stripped_info_copy");
+  SYMS_StrippedInfoArray result = {0};
+  result.count = stripped->count;
+  result.info = syms_push_array(arena, SYMS_StrippedInfo, result.count);
+  
+  SYMS_StrippedInfo *src = stripped->info;
+  SYMS_StrippedInfo *dst = result.info;
+  for (SYMS_U64 i = 0; i < result.count; i += 1, src += 1, dst += 1){
+    dst->name = syms_push_string_copy(arena, src->name);
+    dst->voff = src->voff;
+  }
+  
+  SYMS_ProfEnd();
+  return(result);
+}
+
 #endif //SYMS_DATA_STRUCTURES_C
diff --git a/Engine/Source/Runtime/SymsLib/syms/core/data_structures/syms_data_structures.h b/Engine/Source/Runtime/SymsLib/syms/core/data_structures/syms_data_structures.h
--- a/Engine/Source/Runtime/SymsLib/syms/core/data_structures/syms_data_structures.h
+++ b/Engine/Source/Runtime/SymsLib/syms/core/data_structures/syms_data_structures.h
@@ -248,5 +248,7 @@ SYMS_API SYMS_LineTable    syms_line_table_with_indexes_from_parse(SYMS_Arena *a
 SYMS_API SYMS_String8Array syms_string_array_copy(SYMS_Arena *arena, SYMS_StringCons *cons_optional,
                                                   SYMS_String8Array *array);
 
+SYMS_API SYMS_StrippedInfoArray syms_stripped_info_copy(SYMS_Arena *arena, SYMS_StrippedInfoArray *stripped);
+
 
 #endif //SYMS_DATA_STRUCTURES_H
diff --git a/Engine/Source/Runtime/SymsLib/syms/core/dwarf/syms_dwarf_parser.c b/Engine/Source/Runtime/SymsLib/syms/core/dwarf/syms_dwarf_parser.c
--- a/Engine/Source/Runtime/SymsLib/syms/core/dwarf/syms_dwarf_parser.c
+++ b/Engine/Source/Runtime/SymsLib/syms/core/dwarf/syms_dwarf_parser.c
@@ -242,27 +242,28 @@ syms_dw_symbol_kind_from_tag_stub(SYMS_String8 data, SYMS_DwDbgAccel *dbg, SYMS_
         }
       }
       
-      SYMS_B32 location_is_constant = syms_false;
+      SYMS_B32 location_is_fixed_address = syms_false;
       if(has_location)
       {
         void *info_base = syms_dw_sec_base_from_dbg(data, dbg, SYMS_DwSectionKind_Info);
         SYMS_U64Range loc_rng = syms_make_u64_range(location_value.v[0],
                                                     location_value.v[0] + location_value.v[1]);
-        SYMS_DwLocationDesc loc = syms_dw_location_desc_from_based_range(scratch.arena,
-                                                                         info_base, loc_rng, 0,
-                                                                         resolve_params.addr_size);
-        location_is_constant = (loc.expr_root != 0 && loc.expr_root->first == 0);
+        
+        // TODO(allen): fill this in correctly
+        SYMS_U64 text_base = 0;
+        SYMS_DwSimpleLoc loc = syms_dw_expr__analyze_fast(info_base, loc_rng, text_base);
+        location_is_fixed_address = (loc.kind == SYMS_DwSimpleLocKind_Address);
       }
       
       if(has_const_val)
       {
         symbol_kind = SYMS_SymbolKind_Const;
       }
-      else if(location_is_constant)
+      else if(location_is_fixed_address)
       {
         symbol_kind = SYMS_SymbolKind_ImageRelativeVariable;
       }
-      else if(!has_external && !location_is_constant)
+      else if(!has_external && !location_is_fixed_address)
       {
         symbol_kind = SYMS_SymbolKind_LocalVariable;
       }
@@ -404,7 +405,7 @@ syms_dw_based_range_read_length(void *base, SYMS_U64Range range, SYMS_U64 offset
   if(syms_based_range_read_struct(base, range, offset, &first32))
   {
     // NOTE(rjf): DWARF 32-bit => use the first 32 bits as the size.
-    if(first32 < 0xfffffff0)
+    if(first32 != 0xffffffff)
     {
       value = (SYMS_U64)first32;
       bytes_read = sizeof(SYMS_U32);
@@ -1930,48 +1931,6 @@ syms_dw_tag_from_info_offset(SYMS_Arena *arena,
   return tag;
 }
 
-// TODO(rjf): (2021/06/09) This isn't being used, but I am keeping it around in
-// the short-term in case it becomes useful again soon in the new API
-// implementation.
-SYMS_API SYMS_U64
-syms_dw_parse_tag_children(SYMS_Arena *arena,
-                           SYMS_String8 data,
-                           SYMS_DwDbgAccel *dbg,
-                           SYMS_DwAbbrevTable abbrev_table,
-                           SYMS_U64 address_size,
-                           SYMS_U64 start_info_offset,
-                           SYMS_DwTag *parent)
-{
-  SYMS_U64 bytes_read = 0;
-  if(parent->has_children)
-  {
-    SYMS_DwTag *active_parent = parent;
-    for(SYMS_U64 info_off = start_info_offset;;)
-    {
-      SYMS_DwTag *child = syms_dw_tag_from_info_offset(arena, data, dbg, abbrev_table, address_size, info_off);
-      bytes_read += child->info_range.max - info_off;
-      info_off = child->info_range.max;
-      if(child->abbrev_id == 0)
-      {
-        active_parent = active_parent->parent;
-        if(active_parent == 0)
-        {
-          break;
-        }
-      }
-      SYMS_DwTag *child_store = syms_push_array(arena, SYMS_DwTag, 1);
-      *child_store = *child;
-      child_store->parent = active_parent;
-      SYMS_QueuePush_N(active_parent->first_child, active_parent->last_child, child_store, next_sibling);
-      if(child_store->has_children)
-      {
-        active_parent = child_store;
-      }
-    }
-  }
-  return bytes_read;
-}
-
 SYMS_API SYMS_DwTagStub
 syms_dw_stub_from_tag(SYMS_String8 data, SYMS_DwDbgAccel *dbg, SYMS_DwAttribValueResolveParams resolve_params,
                       SYMS_DwTag *tag)
@@ -2001,879 +1960,6 @@ syms_dw_stub_from_tag(SYMS_String8 data, SYMS_DwDbgAccel *dbg, SYMS_DwAttribValu
   return stub;
 }
 
-////////////////////////////////
-//~ rjf: DWARF Expressions
-
-SYMS_API SYMS_DwExprNode *
-syms_dw_expr_push(SYMS_Arena *arena, SYMS_DwExprEvalState *state)
-{
-  SYMS_DwExprNode *node = syms_push_array(arena, SYMS_DwExprNode, 1);
-  syms_memzero_struct(node);
-  node->next = state->stack_top;
-  state->stack_top = node;
-  return node;
-}
-
-SYMS_API SYMS_DwExprNode *
-syms_dw_expr_push_literal(SYMS_Arena *arena, SYMS_DwExprEvalState *state, SYMS_U64 value)
-{
-  SYMS_DwExprNode *node = syms_dw_expr_push(arena, state);
-  node->value = value;
-  return node;
-}
-
-SYMS_API SYMS_DwExprNode *
-syms_dw_expr_push_operator(SYMS_Arena *arena, SYMS_DwExprEvalState *state, SYMS_DwOpCode opcode, SYMS_DwExprNode *left, SYMS_DwExprNode *right)
-{
-  SYMS_DwExprNode *op = syms_dw_expr_push(arena, state);
-  op->opcode = opcode;
-  if(left != 0)
-  {
-    SYMS_QueuePush(op->first, op->last, left);
-  }
-  if(right != 0)
-  {
-    SYMS_QueuePush(op->first, op->last, right);
-  }
-  return op;
-}
-
-SYMS_API SYMS_DwExprNode *
-syms_dw_expr_push_copy(SYMS_Arena *arena, SYMS_DwExprEvalState *state, SYMS_DwExprNode *to_copy)
-{
-  SYMS_DwExprNode *node = syms_dw_expr_push(arena, state);
-  if(to_copy != 0)
-  {
-    node->value = to_copy->value;
-    node->opcode = to_copy->opcode;
-  }
-  return node;
-}
-
-SYMS_API SYMS_DwExprNode *
-syms_dw_expr_pop(SYMS_DwExprEvalState *state)
-{
-  SYMS_DwExprNode *node = state->stack_top;
-  if(state->stack_top)
-  {
-    state->stack_top = state->stack_top->next;
-  }
-  return node;
-}
-
-SYMS_API SYMS_DwExprNode *
-syms_dw_expr_node_from_stack_index(SYMS_DwExprEvalState *state, SYMS_U64 idx)
-{
-  SYMS_DwExprNode *result = 0;
-  SYMS_U64 i = 0;
-  for(SYMS_DwExprNode *node = state->stack_top; node != 0; node = node->next, i += 1)
-  {
-    if(i == idx)
-    {
-      result = node;
-      break;
-    }
-  }
-  return result;
-}
-
-SYMS_API SYMS_DwLocationDesc
-syms_dw_location_desc_from_based_range(SYMS_Arena *arena,
-                                       void *base,
-                                       SYMS_U64Range range,
-                                       SYMS_U64 offset,
-                                       SYMS_U64 addr_size)
-{
-  SYMS_DwExprEvalState eval_;
-  SYMS_DwExprEvalState *eval = &eval_;
-  syms_memzero_struct(eval);
-  
-  for(;offset < syms_u64_range_size(range);)
-  {
-    SYMS_U8 op8 = 0;
-    offset += syms_based_range_read_struct(base, range, offset, &op8);
-    SYMS_DwOpCode op = (SYMS_DwOpCode)op8;
-    
-    //- rjf: do op
-    switch(op)
-    {
-      case SYMS_DwOpCode_NOP: break;
-      
-      //- rjf: push address onto stack
-      case SYMS_DwOpCode_ADDR:
-      {
-        SYMS_U64 addr = 0;
-        offset += syms_based_range_read(base, range, offset, (SYMS_U64)addr_size, &addr);
-        syms_dw_expr_push_literal(arena, eval, addr);
-      }break;
-      
-      //- rjf: push small literal onto stack
-      case SYMS_DwOpCode_LIT0:  case SYMS_DwOpCode_LIT1:  case SYMS_DwOpCode_LIT2:
-      case SYMS_DwOpCode_LIT3:  case SYMS_DwOpCode_LIT4:  case SYMS_DwOpCode_LIT5:
-      case SYMS_DwOpCode_LIT6:  case SYMS_DwOpCode_LIT7:  case SYMS_DwOpCode_LIT8:
-      case SYMS_DwOpCode_LIT9:  case SYMS_DwOpCode_LIT10: case SYMS_DwOpCode_LIT11:
-      case SYMS_DwOpCode_LIT12: case SYMS_DwOpCode_LIT13: case SYMS_DwOpCode_LIT14:
-      case SYMS_DwOpCode_LIT15: case SYMS_DwOpCode_LIT16: case SYMS_DwOpCode_LIT17:
-      case SYMS_DwOpCode_LIT18: case SYMS_DwOpCode_LIT19: case SYMS_DwOpCode_LIT20:
-      case SYMS_DwOpCode_LIT21: case SYMS_DwOpCode_LIT22: case SYMS_DwOpCode_LIT23:
-      case SYMS_DwOpCode_LIT24: case SYMS_DwOpCode_LIT25: case SYMS_DwOpCode_LIT26:
-      case SYMS_DwOpCode_LIT27: case SYMS_DwOpCode_LIT28: case SYMS_DwOpCode_LIT29:
-      case SYMS_DwOpCode_LIT30: case SYMS_DwOpCode_LIT31:
-      {
-        SYMS_U64 u = (SYMS_U64)(op - SYMS_DwOpCode_LIT0);
-        syms_dw_expr_push_literal(arena, eval, u);
-      }break;
-      
-      //- rjf: push constants onto stack
-      case SYMS_DwOpCode_CONST1U: syms_dw_expr_push_literal(arena, eval, 1); break;
-      case SYMS_DwOpCode_CONST2U: syms_dw_expr_push_literal(arena, eval, 2); break;
-      case SYMS_DwOpCode_CONST4U: syms_dw_expr_push_literal(arena, eval, 4); break;
-      case SYMS_DwOpCode_CONST8U: syms_dw_expr_push_literal(arena, eval, 8); break;
-      case SYMS_DwOpCode_CONST1S: syms_dw_expr_push_literal(arena, eval, 1); break;
-      case SYMS_DwOpCode_CONST2S: syms_dw_expr_push_literal(arena, eval, 2); break;
-      case SYMS_DwOpCode_CONST4S: syms_dw_expr_push_literal(arena, eval, 4); break;
-      case SYMS_DwOpCode_CONST8S: syms_dw_expr_push_literal(arena, eval, 8); break;
-      case SYMS_DwOpCode_CONSTU:
-      {
-        SYMS_U64 v = 0;
-        offset += syms_based_range_read_uleb128(base, range, offset, &v);
-        syms_dw_expr_push_literal(arena, eval, v);
-      }break;
-      
-      //- rjf: push SLEB-128 constant onto stack
-      case SYMS_DwOpCode_CONSTS:
-      {
-        SYMS_S64 v = 0;
-        offset += syms_based_range_read_sleb128(base, range, offset, &v);
-        syms_dw_expr_push_literal(arena, eval, *(SYMS_U64 *)&v);
-      }break;
-      
-      //- rjf: push frame based SLEB-128 onto stack
-      case SYMS_DwOpCode_FBREG:
-      {
-        SYMS_S64 v = 0;
-        offset += syms_based_range_read_sleb128(base, range, offset, &v);
-        SYMS_DwExprNode *v_node = syms_dw_expr_push_literal(arena, eval, *(SYMS_U64 *)&v);
-        syms_dw_expr_pop(eval);
-        syms_dw_expr_push_operator(arena, eval, op, v_node, 0);
-      }break;
-      
-      //- rjf: push register + SLEB-128 onto stack
-      case SYMS_DwOpCode_BREG0:  case SYMS_DwOpCode_BREG1:  case SYMS_DwOpCode_BREG2:
-      case SYMS_DwOpCode_BREG3:  case SYMS_DwOpCode_BREG4:  case SYMS_DwOpCode_BREG5:
-      case SYMS_DwOpCode_BREG6:  case SYMS_DwOpCode_BREG7:  case SYMS_DwOpCode_BREG8:
-      case SYMS_DwOpCode_BREG9:  case SYMS_DwOpCode_BREG10: case SYMS_DwOpCode_BREG11:
-      case SYMS_DwOpCode_BREG12: case SYMS_DwOpCode_BREG13: case SYMS_DwOpCode_BREG14:
-      case SYMS_DwOpCode_BREG15: case SYMS_DwOpCode_BREG16: case SYMS_DwOpCode_BREG17:
-      case SYMS_DwOpCode_BREG18: case SYMS_DwOpCode_BREG19: case SYMS_DwOpCode_BREG20:
-      case SYMS_DwOpCode_BREG21: case SYMS_DwOpCode_BREG22: case SYMS_DwOpCode_BREG23:
-      case SYMS_DwOpCode_BREG24: case SYMS_DwOpCode_BREG25: case SYMS_DwOpCode_BREG26:
-      case SYMS_DwOpCode_BREG27: case SYMS_DwOpCode_BREG28: case SYMS_DwOpCode_BREG29:
-      case SYMS_DwOpCode_BREG30: case SYMS_DwOpCode_BREG31:
-      {
-        SYMS_S64 s = 0;
-        offset += syms_based_range_read_sleb128(base, range, offset, &s);
-        SYMS_DwExprNode *node = syms_dw_expr_push_literal(arena, eval, *(SYMS_U64 *)&s);
-        (void)node;
-        syms_dw_expr_push_operator(arena, eval, op, syms_dw_expr_pop(eval), 0);
-      }break;
-      
-      //- rjf: push ULEB-128 reg onto stack
-      case SYMS_DwOpCode_BREGX:
-      {
-        SYMS_U64 v = 0;
-        offset += syms_based_range_read_uleb128(base, range, offset, &v);
-        SYMS_S64 s = 0;
-        offset += syms_based_range_read_sleb128(base, range, offset, &s);
-        SYMS_DwExprNode *left  = syms_dw_expr_push_literal(arena, eval, v);
-        SYMS_DwExprNode *right = syms_dw_expr_push_literal(arena, eval, *(SYMS_U64 *)&s);
-        syms_dw_expr_pop(eval);
-        syms_dw_expr_pop(eval);
-        syms_dw_expr_push_operator(arena, eval, op, left, right);
-      }break;
-      
-      //- rjf: duplicate + push onto stack
-      case SYMS_DwOpCode_DUP:
-      {
-        SYMS_DwExprNode *top = syms_dw_expr_pop(eval);
-        syms_dw_expr_push_copy(arena, eval, top);
-        syms_dw_expr_push_copy(arena, eval, top);
-      }break;
-      
-      //- rjf: drop top of stack
-      case SYMS_DwOpCode_DROP:    syms_dw_expr_pop(eval); break;
-      
-      //- rjf: stack-picking
-      case SYMS_DwOpCode_OVER:
-      {
-        syms_dw_expr_push_copy(arena, eval, syms_dw_expr_node_from_stack_index(eval, 1));
-      }break;
-      case SYMS_DwOpCode_PICK:
-      {
-        SYMS_U8 index = 0;
-        offset += syms_based_range_read_struct(base, range, offset, &index);
-        syms_dw_expr_push_copy(arena, eval, syms_dw_expr_node_from_stack_index(eval, index));
-      }break;
-      
-      //- rjf: swap
-      case SYMS_DwOpCode_SWAP:
-      {
-        SYMS_DwExprNode *top0 = syms_dw_expr_pop(eval);
-        SYMS_DwExprNode *top1 = syms_dw_expr_pop(eval);
-        syms_dw_expr_push_copy(arena, eval, top0);
-        syms_dw_expr_push_copy(arena, eval, top1);
-      }break;
-      
-      //- rjf: rotate
-      case SYMS_DwOpCode_ROT:
-      {
-        SYMS_DwExprNode *top0 = syms_dw_expr_pop(eval);
-        SYMS_DwExprNode *top1 = syms_dw_expr_pop(eval);
-        SYMS_DwExprNode *top2 = syms_dw_expr_pop(eval);
-        syms_dw_expr_push_copy(arena, eval, top1);
-        syms_dw_expr_push_copy(arena, eval, top0);
-        syms_dw_expr_push_copy(arena, eval, top2);
-      }break;
-      
-      //- rjf: deref size
-      case SYMS_DwOpCode_DEREF_SIZE:
-      {
-        SYMS_U8 read_addr_size = 0;
-        offset += syms_based_range_read_struct(base, range, offset, &read_addr_size);
-        syms_dw_expr_push_operator(arena, eval, op, syms_dw_expr_pop(eval), 0);
-      }break;
-      
-      //- rjf: xderef ???
-      // TODO(nick): This is a very rare case here, maybe later we can figure this out.
-      case SYMS_DwOpCode_XDEREF_SIZE:
-      case SYMS_DwOpCode_XDEREF:
-      {
-        goto exit_fail;
-      }break;
-      
-      //- rjf: pushes of data in the local context, when evaluating
-      case SYMS_DwOpCode_CALL_FRAME_CFA:
-      case SYMS_DwOpCode_PUSH_OBJECT_ADDRESS:
-      {
-        syms_dw_expr_push_operator(arena, eval, op, 0, 0);
-      }break;
-      
-      //- rjf: TLS address
-      case SYMS_DwOpCode_FORM_TLS_ADDRESS:
-      {
-        SYMS_NOT_IMPLEMENTED;
-      }break;
-      
-      //- rjf: unary operators on popped value
-      case SYMS_DwOpCode_ABS:   case SYMS_DwOpCode_NEG: case SYMS_DwOpCode_NOT:
-      case SYMS_DwOpCode_DEREF: 
-      {
-        syms_dw_expr_push_operator(arena, eval, op, syms_dw_expr_pop(eval), 0);
-      }break;
-      
-      //- rjf: binary operators on popped values
-      case SYMS_DwOpCode_PLUS: case SYMS_DwOpCode_DIV: case SYMS_DwOpCode_MINUS:
-      case SYMS_DwOpCode_MOD:  case SYMS_DwOpCode_MUL: case SYMS_DwOpCode_AND:
-      case SYMS_DwOpCode_OR:   case SYMS_DwOpCode_SHL: case SYMS_DwOpCode_SHR:
-      case SYMS_DwOpCode_SHRA: case SYMS_DwOpCode_XOR: case SYMS_DwOpCode_LE:
-      case SYMS_DwOpCode_GE:   case SYMS_DwOpCode_EQ:  case SYMS_DwOpCode_LT:
-      case SYMS_DwOpCode_GT:   case SYMS_DwOpCode_NE:
-      {
-        SYMS_DwExprNode *left = syms_dw_expr_pop(eval);
-        SYMS_DwExprNode *right = syms_dw_expr_pop(eval);
-        syms_dw_expr_push_operator(arena, eval, op, left, right);
-      }break;
-      
-      //- rjf: other operators
-      case SYMS_DwOpCode_PLUS_UCONST:
-      {
-        SYMS_U64 u = 0;
-        offset += syms_based_range_read_uleb128(base, range, offset, &u);
-        SYMS_DwExprNode *left = syms_dw_expr_pop(eval);
-        SYMS_DwExprNode *right = syms_dw_expr_push_literal(arena, eval, u);
-        syms_dw_expr_pop(eval);
-        syms_dw_expr_push_operator(arena, eval, op, left, right);
-      }break;
-      
-      //- rjf: skip
-      case SYMS_DwOpCode_SKIP:
-      {
-        SYMS_S16 s = 0;
-        offset += syms_based_range_read_struct(base, range, offset, &s);
-        offset += s;
-      }break;
-      
-      //- TODO(rjf): not sure what this is
-      case SYMS_DwOpCode_BRA:
-      {
-        SYMS_S16 s = 0;
-        offset += syms_based_range_read_struct(base, range, offset, &s);
-        SYMS_DwExprNode *v = syms_dw_expr_pop(eval);
-        SYMS_S64 sa = *(SYMS_S64 *)&v->value;
-        offset += sa;
-      }break;
-      
-      //- rjf: function calls (wait, are you kidding me?)
-      case SYMS_DwOpCode_CALL_REF:
-      case SYMS_DwOpCode_CALL4:
-      case SYMS_DwOpCode_CALL2:
-      {
-        goto exit_fail;
-      }break;
-      
-      //- rjf: register values
-      case SYMS_DwOpCode_REG0:  case SYMS_DwOpCode_REG1:  case SYMS_DwOpCode_REG2:
-      case SYMS_DwOpCode_REG3:  case SYMS_DwOpCode_REG4:  case SYMS_DwOpCode_REG5:
-      case SYMS_DwOpCode_REG6:  case SYMS_DwOpCode_REG7:  case SYMS_DwOpCode_REG8:
-      case SYMS_DwOpCode_REG9:  case SYMS_DwOpCode_REG10: case SYMS_DwOpCode_REG11:
-      case SYMS_DwOpCode_REG12: case SYMS_DwOpCode_REG13: case SYMS_DwOpCode_REG14:
-      case SYMS_DwOpCode_REG15: case SYMS_DwOpCode_REG16: case SYMS_DwOpCode_REG17:
-      case SYMS_DwOpCode_REG18: case SYMS_DwOpCode_REG19: case SYMS_DwOpCode_REG20:
-      case SYMS_DwOpCode_REG21: case SYMS_DwOpCode_REG22: case SYMS_DwOpCode_REG23:
-      case SYMS_DwOpCode_REG24: case SYMS_DwOpCode_REG25: case SYMS_DwOpCode_REG26:
-      case SYMS_DwOpCode_REG27: case SYMS_DwOpCode_REG28: case SYMS_DwOpCode_REG29:
-      case SYMS_DwOpCode_REG30: case SYMS_DwOpCode_REG31:
-      {
-        syms_dw_expr_push_operator(arena, eval, op, 0, 0);
-      }break;
-      case SYMS_DwOpCode_REGX:
-      {
-        SYMS_U64 reg_index = 0;
-        offset += syms_based_range_read_uleb128(base, range, offset, &reg_index);
-        SYMS_DwExprNode *v = syms_dw_expr_push_literal(arena, eval, reg_index);
-        syms_dw_expr_pop(eval);
-        syms_dw_expr_push_operator(arena, eval, op, v, 0);
-      }break;
-      
-      //- rjf: implicit value
-      case SYMS_DwOpCode_IMPLICIT_VALUE:
-      {
-        SYMS_U64 size = 0;
-        offset += syms_based_range_read_uleb128(base, range, offset, &size);
-        SYMS_DwExprNode *v = syms_dw_expr_push_literal(arena, eval, size);
-        syms_dw_expr_pop(eval);
-        syms_dw_expr_push_operator(arena, eval, op, v, 0);
-        goto exit_good;
-      }break;
-      
-      //- rjf: stack value
-      case SYMS_DwOpCode_STACK_VALUE:
-      {
-        goto exit_fail;
-      }break;
-      
-      //- rjf: pieces ???
-      case SYMS_DwOpCode_PIECE:
-      case SYMS_DwOpCode_BIT_PIECE:
-      {
-        // TODO(rjf): Unimplemented
-        goto exit_fail;
-      }break;
-      
-      //- rjf: catch-all
-      default:
-      {
-        // TODO(rjf): Unexpected opcode
-        goto exit_fail;
-      }break;
-    }
-  }
-  exit_fail:;
-  exit_good:;
-  
-  SYMS_DwLocationDesc result;
-  syms_memzero_struct(&result);
-  result.expr_root = eval->stack_top;
-  return result;
-}
-
-////////////////////////////////
-//~ rjf: Location Evaluation VM
-
-SYMS_API SYMS_U64
-syms_dw_based_range_read_and_eval_location_expr(void *base,
-                                                SYMS_U64Range range,
-                                                SYMS_U64 offset,
-                                                SYMS_DwLocation *location_out,
-                                                SYMS_U64 frame_base,
-                                                SYMS_U64 member_loc,
-                                                SYMS_U64 cfa,
-                                                SYMS_DwMode dw_mode,
-                                                SYMS_Arch arch,
-                                                void *memread_ctx, dw_memread_sig *memread,
-                                                void *regread_ctx, dw_regread_sig *regread)
-{
-  SYMS_U32 addr_size = syms_address_size_from_arch(arch) / 8;
-  SYMS_U64 start_offset = offset;
-  
-  SYMS_DwLocation location;
-  syms_memzero_struct(&location);
-  
-  SYMS_U64 stack[128];
-  SYMS_U64 *end = &stack[SYMS_ARRAY_SIZE(stack) - 1];
-  SYMS_U64 *top = end;
-#define SYMS_DW_PUSH_U64(x)  *top = (SYMS_U64)(x); --top
-#define SYMS_DW_PUSH_S64(x)  *top = (SYMS_S64)(x); --top
-#define SYMS_DW_POP()        (++top)
-#define SYMS_DW_POP_S64()    ((SYMS_S64)*(++top))
-#define SYMS_DW_POP_U64()    (*(++top))
-  
-  for(;offset < syms_u64_range_size(range);)
-  {
-    SYMS_U8 op = 0;
-    offset += syms_based_range_read_struct(base, range, offset, &op);
-    
-    //- rjf: do op
-    switch(op)
-    {
-      case SYMS_DwOpCode_NOP: break;
-      
-      //- rjf: push address
-      case SYMS_DwOpCode_ADDR:
-      {
-        SYMS_U64 addr = 0;
-        offset += syms_based_range_read(base, range, offset, (SYMS_U64)addr_size, &addr);
-        SYMS_DW_PUSH_U64(addr);
-      }break;
-      
-      //- rjf: push small literal
-      case SYMS_DwOpCode_LIT0:  case SYMS_DwOpCode_LIT1:  case SYMS_DwOpCode_LIT2:
-      case SYMS_DwOpCode_LIT3:  case SYMS_DwOpCode_LIT4:  case SYMS_DwOpCode_LIT5:
-      case SYMS_DwOpCode_LIT6:  case SYMS_DwOpCode_LIT7:  case SYMS_DwOpCode_LIT8:
-      case SYMS_DwOpCode_LIT9:  case SYMS_DwOpCode_LIT10: case SYMS_DwOpCode_LIT11:
-      case SYMS_DwOpCode_LIT12: case SYMS_DwOpCode_LIT13: case SYMS_DwOpCode_LIT14:
-      case SYMS_DwOpCode_LIT15: case SYMS_DwOpCode_LIT16: case SYMS_DwOpCode_LIT17:
-      case SYMS_DwOpCode_LIT18: case SYMS_DwOpCode_LIT19: case SYMS_DwOpCode_LIT20:
-      case SYMS_DwOpCode_LIT21: case SYMS_DwOpCode_LIT22: case SYMS_DwOpCode_LIT23:
-      case SYMS_DwOpCode_LIT24: case SYMS_DwOpCode_LIT25: case SYMS_DwOpCode_LIT26:
-      case SYMS_DwOpCode_LIT27: case SYMS_DwOpCode_LIT28: case SYMS_DwOpCode_LIT29:
-      case SYMS_DwOpCode_LIT30: case SYMS_DwOpCode_LIT31:
-      {
-        SYMS_U64 u = (SYMS_U64)(op - SYMS_DwOpCode_LIT0);
-        SYMS_DW_PUSH_U64(u);
-      }break;
-      
-      //- rjf: push constants
-      case SYMS_DwOpCode_CONST1U: SYMS_DW_PUSH_U64(1); break;
-      case SYMS_DwOpCode_CONST2U: SYMS_DW_PUSH_U64(2); break;
-      case SYMS_DwOpCode_CONST4U: SYMS_DW_PUSH_U64(4); break;
-      case SYMS_DwOpCode_CONST8U: SYMS_DW_PUSH_U64(8); break;
-      case SYMS_DwOpCode_CONST1S: SYMS_DW_PUSH_S64(1); break;
-      case SYMS_DwOpCode_CONST2S: SYMS_DW_PUSH_S64(2); break;
-      case SYMS_DwOpCode_CONST4S: SYMS_DW_PUSH_S64(4); break;
-      case SYMS_DwOpCode_CONST8S: SYMS_DW_PUSH_S64(8); break;
-      case SYMS_DwOpCode_CONSTU:
-      {
-        SYMS_U64 v = 0;
-        offset += syms_based_range_read_uleb128(base, range, offset, &v);
-        SYMS_DW_PUSH_U64(v);
-      }break;
-      
-      //- rjf: push SLEB-128 constant
-      case SYMS_DwOpCode_CONSTS:
-      {
-        SYMS_S64 v = 0;
-        offset += syms_based_range_read_sleb128(base, range, offset, &v);
-        SYMS_DW_PUSH_S64(v);
-      }break;
-      
-      //- rjf: push frame based SLEB-128
-      case SYMS_DwOpCode_FBREG:
-      {
-        SYMS_S64 v = 0;
-        offset += syms_based_range_read_sleb128(base, range, offset, &v);
-        v += frame_base;
-        SYMS_DW_PUSH_S64(v);
-      }break;
-      
-      //- rjf: push register + SLEB-128
-      case SYMS_DwOpCode_BREG0:  case SYMS_DwOpCode_BREG1:  case SYMS_DwOpCode_BREG2:
-      case SYMS_DwOpCode_BREG3:  case SYMS_DwOpCode_BREG4:  case SYMS_DwOpCode_BREG5:
-      case SYMS_DwOpCode_BREG6:  case SYMS_DwOpCode_BREG7:  case SYMS_DwOpCode_BREG8:
-      case SYMS_DwOpCode_BREG9:  case SYMS_DwOpCode_BREG10: case SYMS_DwOpCode_BREG11:
-      case SYMS_DwOpCode_BREG12: case SYMS_DwOpCode_BREG13: case SYMS_DwOpCode_BREG14:
-      case SYMS_DwOpCode_BREG15: case SYMS_DwOpCode_BREG16: case SYMS_DwOpCode_BREG17:
-      case SYMS_DwOpCode_BREG18: case SYMS_DwOpCode_BREG19: case SYMS_DwOpCode_BREG20:
-      case SYMS_DwOpCode_BREG21: case SYMS_DwOpCode_BREG22: case SYMS_DwOpCode_BREG23:
-      case SYMS_DwOpCode_BREG24: case SYMS_DwOpCode_BREG25: case SYMS_DwOpCode_BREG26:
-      case SYMS_DwOpCode_BREG27: case SYMS_DwOpCode_BREG28: case SYMS_DwOpCode_BREG29:
-      case SYMS_DwOpCode_BREG30: case SYMS_DwOpCode_BREG31:
-      {
-        SYMS_U64 u = 0;
-        SYMS_U32 reg_index = (SYMS_U32)(op - SYMS_DwOpCode_BREG0);
-        if(!regread(regread_ctx, arch, reg_index, &u, sizeof(u)))
-        {
-          goto exit_fail;
-        }
-        SYMS_S64 s = 0;
-        offset += syms_based_range_read_sleb128(base, range, offset, &s);
-        s += ((SYMS_S64) u);
-        SYMS_DW_PUSH_S64(s);
-      }break;
-      
-      //- rjf: push ULEB-128 reg
-      case SYMS_DwOpCode_BREGX:
-      {
-        SYMS_U64 v = 0;
-        offset += syms_based_range_read_uleb128(base, range, offset, &v);
-        SYMS_U32 reg_index = (SYMS_U32)(v);
-        SYMS_U64 u = 0;
-        if(!regread(regread_ctx, arch, reg_index, &u, sizeof(u)))
-        {
-          goto exit_fail;
-        }
-        SYMS_S64 s = 0;
-        offset += syms_based_range_read_sleb128(base, range, offset, &s);
-        u = ((SYMS_U64) ((SYMS_S64) u) + s);
-        SYMS_DW_PUSH_U64(u);
-      }break;
-      
-      //- rjf: ???
-      case SYMS_DwOpCode_DUP:     SYMS_DW_PUSH_U64(*top);   break;
-      case SYMS_DwOpCode_DROP:    SYMS_DW_POP();            break;
-      case SYMS_DwOpCode_OVER:    SYMS_DW_PUSH_U64(top[1]); break;
-      case SYMS_DwOpCode_PICK:
-      {
-        SYMS_U8 index = 0;
-        offset += syms_based_range_read_struct(base, range, offset, &index);
-        SYMS_DW_PUSH_U64(top[index]);
-      }break;
-      
-      //- rjf: swap
-      case SYMS_DwOpCode_SWAP:
-      {
-        SYMS_U64 u = top[0];
-        top[0] = top[1];
-        top[1] = u;
-      }break;
-      
-      //- rjf: rotate
-      case SYMS_DwOpCode_ROT:
-      {
-        SYMS_U64 u = top[0];
-        top[0] = top[2];
-        top[2] = top[1];
-        top[1] = u;
-      }break;
-      
-      //- rjf: deref
-      case SYMS_DwOpCode_DEREF:
-      {
-        SYMS_ASSERT(sizeof(*top) <= addr_size);
-        SYMS_U64 u = SYMS_DW_POP_U64();
-        if(!memread(memread_ctx, u, &u, addr_size))
-        {
-          goto exit_fail;
-        }
-        SYMS_DW_PUSH_U64(u);
-      }break;
-      
-      //- rjf: deref size
-      case SYMS_DwOpCode_DEREF_SIZE:
-      {
-        SYMS_U8 read_addr_size = 0;
-        offset += syms_based_range_read_struct(base, range, offset, &read_addr_size);
-        if (read_addr_size > addr_size)
-        {
-          goto exit_fail;
-        }
-        SYMS_U64 u = SYMS_DW_POP_U64();
-        if (!memread(memread_ctx, u, &u, read_addr_size))
-        {
-          goto exit_fail;
-        }
-        SYMS_DW_PUSH_U64(u);
-      }break;
-      
-      //- rjf: xderef ???
-      // TODO(nick): This is a very rare case here, maybe later we can figure this out.
-      case SYMS_DwOpCode_XDEREF_SIZE:
-      case SYMS_DwOpCode_XDEREF:
-      {
-        goto exit_fail;
-      }break;
-      
-      //- rjf: push member loc
-      case SYMS_DwOpCode_PUSH_OBJECT_ADDRESS:
-      {
-        SYMS_DW_PUSH_U64(member_loc);
-      }break;
-      
-      //- rjf: TLS address
-      case SYMS_DwOpCode_FORM_TLS_ADDRESS:
-      {
-        SYMS_NOT_IMPLEMENTED;
-      }break;
-      
-      case SYMS_DwOpCode_CALL_FRAME_CFA:
-      {
-        SYMS_DW_PUSH_U64(cfa);
-      }break;
-      
-      //- rjf: arithmetic ops
-      case SYMS_DwOpCode_ABS:
-      {
-        SYMS_S64 s = SYMS_DW_POP_S64();
-        SYMS_U64 u = ((SYMS_U64) s < 0 ? -s : s);
-        SYMS_DW_PUSH_U64(u);
-      }break;
-      case SYMS_DwOpCode_PLUS:
-      {
-        SYMS_U64 ua = SYMS_DW_POP_U64();
-        SYMS_U64 ub = SYMS_DW_POP_U64();
-        SYMS_U64 u = ub + ua;
-        SYMS_DW_PUSH_U64(u);
-      }break;
-      case SYMS_DwOpCode_PLUS_UCONST:
-      {
-        SYMS_U64 u = 0;
-        offset += syms_based_range_read_uleb128(base, range, offset, &u);
-        u += SYMS_DW_POP_U64();
-        SYMS_DW_PUSH_U64(u);
-      }break;
-      case SYMS_DwOpCode_DIV:
-      {
-        SYMS_S64 sa = SYMS_DW_POP_S64();
-        SYMS_S64 sb = SYMS_DW_POP_S64();
-        if(sa == 0)
-        {
-          goto exit_fail;
-        }
-        SYMS_S64 s = sb / sa;
-        SYMS_DW_PUSH_S64(s);
-      }break;
-      case SYMS_DwOpCode_MINUS:
-      {
-        SYMS_S64 sa = SYMS_DW_POP_S64();
-        SYMS_S64 sb = SYMS_DW_POP_S64();
-        SYMS_S64 s = sb - sa;
-        SYMS_DW_PUSH_S64(s);
-      }break;
-      case SYMS_DwOpCode_MOD:
-      {
-        SYMS_S64 sa = SYMS_DW_POP_S64();
-        SYMS_S64 sb = SYMS_DW_POP_S64();
-        SYMS_S64 s = sb % sa;
-        SYMS_DW_PUSH_S64(s);
-      }break;
-      case SYMS_DwOpCode_MUL:
-      {
-        SYMS_S64 sa = SYMS_DW_POP_S64();
-        SYMS_S64 sb = SYMS_DW_POP_S64();
-        if(sa == 0)
-        {
-          goto exit_fail;
-        }
-        SYMS_S64 s = sb * sa;
-        SYMS_DW_PUSH_S64(s);
-      }break;
-      case SYMS_DwOpCode_NEG:
-      {
-        SYMS_S64 s = SYMS_DW_POP_S64();
-        s = -s;
-        SYMS_DW_PUSH_S64(s);
-      }break;
-      
-      //- rjf: boolean/bitwise ops
-      case SYMS_DwOpCode_AND:
-      {
-        SYMS_S64 sa = SYMS_DW_POP_S64();
-        SYMS_S64 sb = SYMS_DW_POP_S64();
-        SYMS_S64 s = sb & sa;
-        SYMS_DW_PUSH_S64(s);
-      }break;
-      case SYMS_DwOpCode_NOT:
-      {
-        SYMS_U64 u = SYMS_DW_POP_U64();
-        u = ((SYMS_U64) !u);
-        SYMS_DW_PUSH_U64(u);
-      }break;
-      case SYMS_DwOpCode_OR:
-      {
-        SYMS_U64 ua = SYMS_DW_POP_U64();
-        SYMS_U64 ub = SYMS_DW_POP_U64();
-        SYMS_U64 u = ub | ua;
-        SYMS_DW_PUSH_U64(u);
-      }break;
-      case SYMS_DwOpCode_SHL:
-      {
-        SYMS_U64 ua = SYMS_DW_POP_U64();
-        SYMS_U64 ub = SYMS_DW_POP_U64();
-        SYMS_U64 u = ub << ua;
-        SYMS_DW_PUSH_U64(u);
-      }break;
-      case SYMS_DwOpCode_SHR:
-      {
-        SYMS_U64 ua = SYMS_DW_POP_U64();
-        SYMS_U64 ub = SYMS_DW_POP_U64();
-        SYMS_U64 u = ub >> ua;
-        SYMS_DW_PUSH_U64(u);
-      }break;
-      case SYMS_DwOpCode_SHRA:
-      {
-        SYMS_S64 sa = SYMS_DW_POP_S64();
-        SYMS_S64 sb = SYMS_DW_POP_S64();
-        SYMS_S64 s = sb >> sa;
-        SYMS_DW_PUSH_S64(s);
-      }break;
-      case SYMS_DwOpCode_XOR:
-      {
-        SYMS_U64 ua = SYMS_DW_POP_U64();
-        SYMS_U64 ub = SYMS_DW_POP_U64();
-        SYMS_U64 u = ub ^ ua;
-        SYMS_DW_PUSH_U64(u);
-      }break;
-      case SYMS_DwOpCode_LE:
-      {
-        SYMS_S64 sa = SYMS_DW_POP_S64();
-        SYMS_S64 sb = SYMS_DW_POP_S64();
-        SYMS_S64 s = sb <= sa;
-        SYMS_DW_PUSH_S64(s);
-      }break;
-      case SYMS_DwOpCode_GE:
-      {
-        SYMS_S64 sa = SYMS_DW_POP_S64();
-        SYMS_S64 sb = SYMS_DW_POP_S64();
-        SYMS_S64 s = sb >= sa;
-        SYMS_DW_PUSH_S64(s);
-      }break;
-      case SYMS_DwOpCode_EQ:
-      {
-        SYMS_S64 sa = SYMS_DW_POP_S64();
-        SYMS_S64 sb = SYMS_DW_POP_S64();
-        SYMS_S64 s = sb == sa;
-        SYMS_DW_PUSH_S64(s);
-      }break;
-      case SYMS_DwOpCode_LT:
-      {
-        SYMS_S64 sa = SYMS_DW_POP_S64();
-        SYMS_S64 sb = SYMS_DW_POP_S64();
-        SYMS_S64 s = sb < sa;
-        SYMS_DW_PUSH_S64(s);
-      }break;
-      case SYMS_DwOpCode_GT:
-      {
-        SYMS_S64 sa = SYMS_DW_POP_S64();
-        SYMS_S64 sb = SYMS_DW_POP_S64();
-        SYMS_S64 s = sb > sa;
-        SYMS_DW_PUSH_S64(s);
-      }break;
-      case SYMS_DwOpCode_NE:
-      {
-        SYMS_S64 sa = SYMS_DW_POP_S64();
-        SYMS_S64 sb = SYMS_DW_POP_S64();
-        SYMS_S64 s = sb != sa;
-        SYMS_DW_PUSH_S64(s);
-      }break;
-      case SYMS_DwOpCode_SKIP:
-      {
-        SYMS_S16 s = 0;
-        offset += syms_based_range_read_struct(base, range, offset, &s);
-        offset += s;
-      }break;
-      case SYMS_DwOpCode_BRA:
-      {
-        SYMS_S16 s = 0;
-        offset += syms_based_range_read_struct(base, range, offset, &s);
-        SYMS_S64 sa = SYMS_DW_POP_S64();
-        if(sa != 0)
-        {
-          offset += sa;
-        }
-      }break;
-      
-      //- rjf: function calls (wait, are you kidding me?)
-      case SYMS_DwOpCode_CALL_REF:
-      case SYMS_DwOpCode_CALL4:
-      case SYMS_DwOpCode_CALL2:
-      {
-        goto exit_fail;
-      }break;
-      
-      case SYMS_DwOpCode_REG0:  case SYMS_DwOpCode_REG1:  case SYMS_DwOpCode_REG2:
-      case SYMS_DwOpCode_REG3:  case SYMS_DwOpCode_REG4:  case SYMS_DwOpCode_REG5:
-      case SYMS_DwOpCode_REG6:  case SYMS_DwOpCode_REG7:  case SYMS_DwOpCode_REG8:
-      case SYMS_DwOpCode_REG9:  case SYMS_DwOpCode_REG10: case SYMS_DwOpCode_REG11:
-      case SYMS_DwOpCode_REG12: case SYMS_DwOpCode_REG13: case SYMS_DwOpCode_REG14:
-      case SYMS_DwOpCode_REG15: case SYMS_DwOpCode_REG16: case SYMS_DwOpCode_REG17:
-      case SYMS_DwOpCode_REG18: case SYMS_DwOpCode_REG19: case SYMS_DwOpCode_REG20:
-      case SYMS_DwOpCode_REG21: case SYMS_DwOpCode_REG22: case SYMS_DwOpCode_REG23:
-      case SYMS_DwOpCode_REG24: case SYMS_DwOpCode_REG25: case SYMS_DwOpCode_REG26:
-      case SYMS_DwOpCode_REG27: case SYMS_DwOpCode_REG28: case SYMS_DwOpCode_REG29:
-      case SYMS_DwOpCode_REG30: case SYMS_DwOpCode_REG31:
-      {
-        SYMS_U64 u = 0;
-        SYMS_U32 reg_index = ((SYMS_U32) op - SYMS_DwOpCode_REG0);
-        if(!regread(regread_ctx, arch, reg_index, &u, sizeof(u)))
-        {
-          goto exit_fail;
-        }
-        SYMS_DW_PUSH_U64(u);
-      }break;
-      
-      case SYMS_DwOpCode_REGX:
-      {
-        SYMS_U64 u = 0;
-        SYMS_U64 reg_index = 0;
-        offset += syms_based_range_read_uleb128(base, range, offset, &reg_index);
-        reg_index = (SYMS_U32)(reg_index);
-        if(!regread(regread_ctx, arch, reg_index, &u, sizeof(u)))
-        {
-          goto exit_fail;
-        }
-        
-        SYMS_DW_PUSH_U64(u);
-      }break;
-      
-      case SYMS_DwOpCode_IMPLICIT_VALUE:
-      {
-        SYMS_U64 len = 0;
-        offset += syms_based_range_read_uleb128(base, range, offset, &len);
-        location.implicit_value = syms_based_range_ptr(base, range, offset);
-        location.implicit_value_size = len;
-        goto exit_good;
-      }break;
-      
-      case SYMS_DwOpCode_STACK_VALUE:
-      {
-        goto exit_fail;
-      }break;
-      
-      case SYMS_DwOpCode_PIECE:
-      case SYMS_DwOpCode_BIT_PIECE:
-      {
-        SYMS_NOT_IMPLEMENTED;
-      }break;
-      
-      default:
-      {
-        SYMS_ASSERT(!"encountered an unimplemented expression opcode");
-      }break;
-    }
-  }
-  
-  exit_good:;
-  if(top != end && offset >= syms_u64_range_size(range))
-  {
-    location.va = SYMS_DW_POP_U64();
-  }
-  exit_fail:;
-  SYMS_U64 result = offset - start_offset;
-  
-#undef SYMS_DW_PUSH_U64
-#undef SYMS_DW_PUSH_S64
-#undef SYMS_DW_POP
-#undef SYMS_DW_POP_S64
-#undef SYMS_DW_POP_U64
-  
-  if(location_out)
-  {
-    *location_out = location;
-  }
-  
-  return result;
-}
-
 ////////////////////////////////
 //~ rjf: Unit Set Accel;
 
@@ -3104,7 +2190,6 @@ syms_dw_comp_root_from_range(SYMS_Arena *arena, SYMS_String8 data, SYMS_DwDbgAcc
   SYMS_B32        high_pc_is_relative = 0;
   SYMS_DwAttribValue ranges_attrib_value = {SYMS_DwSectionKind_Null};
   SYMS_U64        line_base         = 0;
-  SYMS_DwAttrib *high_pc_attrib = 0;
   if(got_comp_unit_tag)
   {
     for(SYMS_DwAttribNode *attrib_n = comp_unit_tag->attribs.first; attrib_n; attrib_n = attrib_n->next)
@@ -3694,8 +2779,6 @@ syms_dw_unit_accel_from_comp_root(SYMS_Arena *arena, SYMS_String8 data, SYMS_DwD
   unit->abbrev_table.count = comp_root->abbrev_table.count;
   unit->abbrev_table.entries = (SYMS_DwAbbrevTableEntry *)syms_push_string_copy(arena, syms_str8((SYMS_U8 *)comp_root->abbrev_table.entries,
                                                                                                  sizeof(SYMS_DwAbbrevTableEntry) * unit->abbrev_table.count)).str;
-  // TODO(rjf): This is an arbitrary choice of size. Can we predict it? Or tune it to
-  // common cases at least?
   unit->stub_table_size = syms_dw_predict_good_stub_table_size_from_range_size(comp_root->tags_info_range.max - comp_root->tags_info_range.min);
   unit->stub_table = syms_push_array_zero(arena, SYMS_DwTagStubCacheNode *, unit->stub_table_size);
   // TODO(rjf): This is an arbitrary choice of size. Can we predict it? Or tune it to
@@ -4329,8 +3412,6 @@ syms_dw_attrib_list_from_stub(SYMS_Arena *arena, SYMS_String8 data, SYMS_DwDbgAc
 SYMS_API SYMS_SymbolIDArray
 syms_dw_copy_sid_array_if_needed(SYMS_Arena *arena, SYMS_SymbolIDArray arr)
 {
-  // TODO(rjf): we can check here if `arena` is the same one we've built the cached array on,
-  // in which case we can avoid the copy
   SYMS_SymbolIDArray result = {0};
   result.count = arr.count;
   result.ids = syms_push_array(arena, SYMS_SymbolID, result.count);
@@ -4477,12 +3558,15 @@ syms_dw_mems_accel_from_sid(SYMS_Arena *arena, SYMS_String8 data, SYMS_DwDbgAcce
                 void *debug_info_base = syms_dw_sec_base_from_dbg(data, dbg, SYMS_DwSectionKind_Info);
                 void *exprloc_base = (void *)((SYMS_U8 *)debug_info_base + value.v[0]);
                 SYMS_U64Range exprloc_rng = syms_make_u64_range(0, value.v[1]);
-                SYMS_DwLocationDesc loc = syms_dw_location_desc_from_based_range(scratch.arena, exprloc_base, exprloc_rng, 0, unit->address_size);
-                if(loc.expr_root != 0 && loc.expr_root->first == 0)
+                
+                // TODO(allen): fill this in correctly
+                SYMS_U64 text_base = 0;
+                SYMS_DwSimpleLoc loc = syms_dw_expr__analyze_fast(exprloc_base, exprloc_rng, text_base);
+                if(loc.kind == SYMS_DwSimpleLocKind_Address)
                 {
                   // NOTE(rjf): These are reported in slots, but PDB reports in byte offsets,
                   // so this maps to the same space as PDB.
-                  virtual_offset = loc.expr_root->value * unit->address_size;
+                  virtual_offset = loc.addr * unit->address_size;
                 }
                 syms_release_scratch(scratch);
               }break;
@@ -5152,18 +4236,19 @@ syms_dw_voff_from_var_sid(SYMS_String8 data, SYMS_DwDbgAccel *dbg, SYMS_DwUnitAc
         
         default:
         {
-          SYMS_ArenaTemp scratch2 = syms_arena_temp_begin(scratch.arena);
+          SYMS_ArenaTemp temp2 = syms_arena_temp_begin(scratch.arena);
           void *debug_info_base = syms_dw_sec_base_from_dbg(data, dbg, SYMS_DwSectionKind_Info);
           void *exprloc_base = ((SYMS_U8*)debug_info_base + value.v[0]);
           SYMS_U64Range exprloc_range = syms_make_u64_range(0, value.v[1]);
-          SYMS_DwLocationDesc loc = syms_dw_location_desc_from_based_range(scratch.arena,
-                                                                           exprloc_base, exprloc_range,
-                                                                           0, unit->address_size);
-          if (loc.expr_root != 0 && loc.expr_root->first == 0){
-            result = loc.expr_root->value;
+          
+          // TODO(allen): fill this in correctly
+          SYMS_U64 text_base = 0;
+          SYMS_DwSimpleLoc loc = syms_dw_expr__analyze_fast(exprloc_base, exprloc_range, text_base);
+          if (loc.kind == SYMS_DwSimpleLocKind_Address){
+            result = loc.addr;
           }
           
-          syms_arena_temp_end(scratch2);
+          syms_arena_temp_end(temp2);
         }break;
       }
     }
@@ -5434,55 +4519,6 @@ syms_dw_scope_children_from_sid(SYMS_Arena *arena, SYMS_String8 data, SYMS_DwDbg
   return array;
 }
 
-//- rjf: location descriptions
-
-SYMS_API SYMS_DwLocationDesc
-syms_dw_location_desc_from_sid(SYMS_Arena *arena, SYMS_String8 data, SYMS_DwDbgAccel *dbg,
-                               SYMS_DwUnitSetAccel *unit_set, SYMS_DwUnitAccel *unit, SYMS_SymbolID sid)
-{
-  SYMS_ArenaTemp scratch = syms_get_scratch(&arena, 1);
-  SYMS_DwLocationDesc desc = {0};
-  SYMS_DwTagStub stub = syms_dw_cached_tag_stub_from_sid__parse_fallback(data, dbg, unit, sid);
-  
-  if(stub.sid != 0)
-  {
-    //- rjf: get attribs
-    SYMS_DwAttribList attribs = syms_dw_attrib_list_from_stub(scratch.arena, data, dbg, unit->address_size, &stub);
-    
-    //- rjf: get comp unit
-    SYMS_DwCompRoot *root = syms_dw_comp_root_from_uid(unit_set, unit->uid);
-    
-    //- rjf: check attribs, see if we have a location
-    for(SYMS_DwAttribNode *attrib_n = attribs.first;
-        attrib_n != 0;
-        attrib_n = attrib_n->next)
-    {
-      SYMS_DwAttrib *attrib = &attrib_n->attrib;
-      if(attrib->attrib_kind == SYMS_DwAttribKind_LOCATION)
-      {
-        SYMS_DwAttribValue value = syms_dw_attrib_value_from_form_value(data, dbg, unit->resolve_params, attrib->form_kind, attrib->value_class, attrib->form_value);
-        
-        if(attrib->form_kind == SYMS_DwFormKind_EXPRLOC)
-        {
-          SYMS_U64 offset = value.v[0];
-          void *base = syms_dw_sec_base_from_dbg(data, dbg, SYMS_DwSectionKind_Info);
-          SYMS_U64Range range = syms_make_u64_range(value.v[0], value.v[0]+ value.v[1]);
-          desc = syms_dw_location_desc_from_based_range(arena, base, range, offset, root->address_size);
-        }
-        else if(attrib->form_kind == SYMS_DwFormKind_LOCLISTX)
-        {
-          desc.va = value.v[0];
-        }
-        
-        break;
-      }
-    }
-  }
-  
-  syms_release_scratch(scratch);
-  return desc;
-}
-
 //- rjf: line info
 
 SYMS_API void
diff --git a/Engine/Source/Runtime/SymsLib/syms/core/dwarf/syms_dwarf_parser.h b/Engine/Source/Runtime/SymsLib/syms/core/dwarf/syms_dwarf_parser.h
--- a/Engine/Source/Runtime/SymsLib/syms/core/dwarf/syms_dwarf_parser.h
+++ b/Engine/Source/Runtime/SymsLib/syms/core/dwarf/syms_dwarf_parser.h
@@ -104,17 +104,6 @@ struct SYMS_DwDbgAccel
   SYMS_B32 is_dwo;
 };
 
-////////////////////////////////
-//~ rjf: Parser Helpers
-
-typedef struct SYMS_DwSegOffArray{
-  SYMS_U32 segoff_size;
-  SYMS_U32 segsel_size;
-  SYMS_U32 num;
-  SYMS_DwSectionKind section;
-  SYMS_U64 section_off;
-} SYMS_DwSegOffArray;
-
 ////////////////////////////////
 //~ rjf: Basic Line Info
 
@@ -363,56 +352,6 @@ struct SYMS_DwTagStubList
   SYMS_U64 count;
 };
 
-////////////////////////////////
-//~ rjf: DWARF Expressions
-
-typedef struct SYMS_DwExprNode SYMS_DwExprNode;
-struct SYMS_DwExprNode
-{
-  SYMS_DwExprNode *first;
-  SYMS_DwExprNode *last;
-  SYMS_DwExprNode *next;
-  SYMS_DwOpCode opcode;
-  SYMS_U64 value;
-};
-
-typedef struct SYMS_DwExprEvalState SYMS_DwExprEvalState;
-struct SYMS_DwExprEvalState
-{
-  SYMS_DwExprNode *stack_top;
-};
-
-////////////////////////////////
-//~ rjf: Locations
-
-typedef struct SYMS_DwLocationDesc SYMS_DwLocationDesc;
-struct SYMS_DwLocationDesc
-{
-  SYMS_U64 va;
-  SYMS_DwExprNode *expr_root;
-};
-
-typedef struct SYMS_DwLocation SYMS_DwLocation;
-struct SYMS_DwLocation
-{
-  SYMS_U64 va;
-  SYMS_U64 implicit_value_size;
-  void *implicit_value;
-};
-
-// NOTE(nick): This callback is used during decoding of the addresses(e.g. locals, globals, lines).
-// If user chose to omit this callback (i.e. passes a NULL) library won't be able to resolve addresses.
-#define DW_REGREAD_SIG(name) SYMS_U32 name(void *context, SYMS_Arch arch, SYMS_U32 reg_index, void *read_buffer, SYMS_U32 read_buffer_max)
-typedef DW_REGREAD_SIG(dw_regread_sig);
-
-#define DW_REGWRITE_SIG(name) SYMS_U32 name(void *context, SYMS_Arch arch, SYMS_U32 reg_index, void *value, SYMS_U32 value_size)
-typedef DW_REGWRITE_SIG(dw_regwrite_sig);
-
-// NOTE(nick): This callback is primarily used for decoding variables' addresses. 
-// If user chose to omit this callback (i.e. passes a NULL) library won't be able to resolve addresses.
-#define DW_MEMREAD_SIG(name) SYMS_B32 name(void *context, SYMS_U64 va, void *read_buffer, SYMS_U32 num_read)
-typedef DW_MEMREAD_SIG(dw_memread_sig);
-
 ////////////////////////////////
 //~ rjf: Line Info VM Types
 
@@ -708,48 +647,9 @@ SYMS_API SYMS_DwTag *syms_dw_tag_from_info_offset(SYMS_Arena *arena,
                                                   SYMS_DwAbbrevTable abbrev_table,
                                                   SYMS_U64 address_size,
                                                   SYMS_U64 info_offset);
-SYMS_API SYMS_U64 syms_dw_parse_tag_children(SYMS_Arena *arena,
-                                             SYMS_String8 data,
-                                             SYMS_DwDbgAccel *dbg,
-                                             SYMS_DwAbbrevTable abbrev_table,
-                                             SYMS_U64 address_size,
-                                             SYMS_U64 start_info_offset,
-                                             SYMS_DwTag *parent);
 SYMS_API SYMS_DwTagStub syms_dw_stub_from_tag(SYMS_String8 data, SYMS_DwDbgAccel *dbg, SYMS_DwAttribValueResolveParams resolve_params,
                                               SYMS_DwTag *tag);
 
-////////////////////////////////
-//~ rjf: DWARF Expressions
-
-SYMS_API SYMS_DwExprNode *syms_dw_expr_push(SYMS_Arena *arena, SYMS_DwExprEvalState *state);
-SYMS_API SYMS_DwExprNode *syms_dw_expr_push_literal(SYMS_Arena *arena, SYMS_DwExprEvalState *state,
-                                                    SYMS_U64 value);
-SYMS_API SYMS_DwExprNode *syms_dw_expr_push_operator(SYMS_Arena *arena, SYMS_DwExprEvalState *state,
-                                                     SYMS_DwOpCode opcode, SYMS_DwExprNode *left,
-                                                     SYMS_DwExprNode *right);
-SYMS_API SYMS_DwExprNode *syms_dw_expr_push_copy(SYMS_Arena *arena, SYMS_DwExprEvalState *state,
-                                                 SYMS_DwExprNode *to_copy);
-SYMS_API SYMS_DwExprNode *syms_dw_expr_pop(SYMS_DwExprEvalState *state);
-SYMS_API SYMS_DwExprNode *syms_dw_expr_node_from_stack_index(SYMS_DwExprEvalState *state, SYMS_U64 idx);
-SYMS_API SYMS_DwLocationDesc syms_dw_location_desc_from_based_range(SYMS_Arena *arena, void *base,
-                                                                    SYMS_U64Range range, SYMS_U64 offset,
-                                                                    SYMS_U64 addr_size);
-
-////////////////////////////////
-//~ rjf: Location Evaluation VM
-
-SYMS_API SYMS_U64 syms_dw_based_range_read_and_eval_location_expr(void *base,
-                                                                  SYMS_U64Range range,
-                                                                  SYMS_U64 offset,
-                                                                  SYMS_DwLocation *location_out,
-                                                                  SYMS_U64 frame_base,
-                                                                  SYMS_U64 member_loc,
-                                                                  SYMS_U64 cfa,
-                                                                  SYMS_DwMode dw_mode,
-                                                                  SYMS_Arch arch,
-                                                                  void *memread_ctx, dw_memread_sig *memread,
-                                                                  void *regread_ctx, dw_regread_sig *regread);
-
 ////////////////////////////////
 //~ rjf: Unit Set Accelerator
 
@@ -769,6 +669,8 @@ SYMS_API SYMS_U64            syms_dw_unit_number_from_uid(SYMS_DwUnitSetAccel *u
 SYMS_API SYMS_UnitInfo       syms_dw_unit_info_from_uid(SYMS_DwUnitSetAccel *unit_set, SYMS_UnitID uid);
 SYMS_API SYMS_UnitNames      syms_dw_unit_names_from_uid(SYMS_Arena *arena, SYMS_DwUnitSetAccel *unit_set,
                                                          SYMS_UnitID uid);
+SYMS_API void                syms_dw_sort_unit_range_point_array_in_place__merge(SYMS_DwUnitRangePoint *a, SYMS_U64 left, SYMS_U64 right, SYMS_U64 end, SYMS_DwUnitRangePoint *b);
+SYMS_API void                syms_dw_sort_unit_range_point_array_in_place(SYMS_DwUnitRangePoint *a, SYMS_U64 count);
 SYMS_API SYMS_UnitRangeArray syms_dw_unit_ranges_from_set(SYMS_Arena *arena,
                                                           SYMS_String8 data, SYMS_DwDbgAccel *dbg,
                                                           SYMS_DwUnitSetAccel *unit_set);
@@ -865,11 +767,6 @@ SYMS_API SYMS_UnitIDAndSig syms_dw_proc_sig_handle_from_sid(SYMS_String8 data, S
 SYMS_API SYMS_SigInfo      syms_dw_sig_info_from_handle(SYMS_Arena *arena, SYMS_String8 data, SYMS_DwDbgAccel *dbg, SYMS_DwUnitAccel *unit, SYMS_SigHandle handle);
 SYMS_API SYMS_SymbolIDArray syms_dw_scope_children_from_sid(SYMS_Arena *arena, SYMS_String8 data, SYMS_DwDbgAccel *dbg, SYMS_DwUnitAccel *unit, SYMS_SymbolID id);
 
-//- rjf: location descriptions
-SYMS_API SYMS_DwLocationDesc syms_dw_location_desc_from_sid(SYMS_Arena *arena, SYMS_String8 data,
-                                                            SYMS_DwDbgAccel *dbg, SYMS_DwUnitSetAccel *unit_set,
-                                                            SYMS_DwUnitAccel *unit, SYMS_SymbolID sid);
-
 //- rjf: line info
 SYMS_API void              syms_dw_line_vm_reset(SYMS_DwLineVMState *state, SYMS_B32 default_is_stmt);
 SYMS_API void              syms_dw_line_vm_advance(SYMS_DwLineVMState *state, SYMS_U64 advance,
diff --git a/Engine/Source/Runtime/SymsLib/syms/core/elf/syms_elf.c b/Engine/Source/Runtime/SymsLib/syms/core/elf/syms_elf.c
--- a/Engine/Source/Runtime/SymsLib/syms/core/elf/syms_elf.c
+++ b/Engine/Source/Runtime/SymsLib/syms/core/elf/syms_elf.c
@@ -3,7 +3,12 @@
 #define SYMS_ELF_C
 
 ////////////////////////////////
-// NOTE(allen): ELF Format Functions
+//~ NOTE(nick): Generated
+
+#include "syms/core/generated/syms_meta_elf.c"
+
+////////////////////////////////
+//~ rjf: 32 => 64 bit conversions
 
 SYMS_API SYMS_ElfEhdr64
 syms_elf_ehdr64_from_ehdr32(SYMS_ElfEhdr32 h32){
@@ -58,6 +63,53 @@ syms_elf_phdr64_from_phdr32(SYMS_ElfPhdr32 h32){
   return(h64);
 }
 
+SYMS_API SYMS_ElfDyn64
+syms_elf_dyn64_from_dyn32(SYMS_ElfDyn32 h32){
+  SYMS_ElfDyn64 h64;
+  h64.tag = (SYMS_U64)h32.tag;
+  h64.val = (SYMS_U64)h32.val;
+  return h64;
+}
+
+SYMS_API SYMS_ElfSym64
+syms_elf_sym64_from_sym32(SYMS_ElfSym32 sym32)
+{
+  SYMS_ElfSym64 sym64;
+  sym64.st_name  = sym32.st_name;
+  sym64.st_value = sym32.st_value;
+  sym64.st_size  = sym32.st_size;
+  sym64.st_info  = sym32.st_info;
+  sym64.st_other = sym32.st_other;
+  sym64.st_shndx = sym32.st_shndx;
+  return sym64;
+}
+
+SYMS_API SYMS_ElfRel64
+syms_elf_rel64_from_rel32(SYMS_ElfRel32 rel32)
+{
+  SYMS_U32 sym = SYMS_ELF32_R_SYM(rel32.r_info);
+  SYMS_U32 type = SYMS_ELF32_R_TYPE(rel32.r_info);
+  SYMS_ElfRel64 rel64;
+  rel64.r_info = SYMS_ELF64_R_INFO(sym, type);
+  rel64.r_offset = rel32.r_offset;
+  return rel64;
+}
+
+SYMS_API SYMS_ElfRela64
+syms_elf_rela64_from_rela32(SYMS_ElfRela32 rela32)
+{
+  SYMS_U32 sym = SYMS_ELF32_R_SYM(rela32.r_info);
+  SYMS_U32 type = SYMS_ELF32_R_TYPE(rela32.r_info);
+  SYMS_ElfRela64 rela64;
+  rela64.r_offset = rela32.r_info;
+  rela64.r_info = SYMS_ELF64_R_INFO(sym, type);
+  rela64.r_addend = rela32.r_addend;
+  return rela64;
+}
+
+////////////////////////////////
+//~ rjf: .gnu_debuglink section 32-bit CRC
+
 // NOTE(rjf): see: https://sourceware.org/gdb/current/onlinedocs/gdb/Separate-Debug-Files.html
 SYMS_API SYMS_U32
 syms_elf_gnu_debuglink_crc32(SYMS_U32 crc, SYMS_String8 data)
diff --git a/Engine/Source/Runtime/SymsLib/syms/core/elf/syms_elf.h b/Engine/Source/Runtime/SymsLib/syms/core/elf/syms_elf.h
--- a/Engine/Source/Runtime/SymsLib/syms/core/elf/syms_elf.h
+++ b/Engine/Source/Runtime/SymsLib/syms/core/elf/syms_elf.h
@@ -4,80 +4,31 @@
 #ifndef SYMS_ELF_H
 #define SYMS_ELF_H
 
-//- rjf: p_type
-typedef enum SYMS_ElfPKind{
-  SYMS_ElfPKind_Null      = 0u,
-  SYMS_ElfPKind_Load      = 1u,
-  SYMS_ElfPKind_Dynamic   = 2u,
-  SYMS_ElfPKind_Interp    = 3u,
-  SYMS_ElfPKind_Note      = 4u,
-  SYMS_ElfPKind_SHLib     = 5u,
-  SYMS_ElfPKind_PHDR      = 6u,
-  SYMS_ElfPKind_LowProc   = 0x70000000u,
-  SYMS_ElfPKind_HighProc  = 0x7fffffffu,
-  
-  //- specific to Sun
-  SYMS_ElfPKind_LowSunW   = 0x6ffffffau,
-  SYMS_ElfPKind_SunWBSS   = 0x6ffffffbu,
-  SYMS_ElfPKind_SunWStack = 0x6ffffffau,
-  SYMS_ElfPKind_HighSunW  = 0x6fffffffu,
-}SYMS_ElfPKind;
-
-//- rjf: e_machine
-typedef enum SYMS_ElfMachineKind{
-  SYMS_ElfMachineKind_NONE           = 0u, 
-  SYMS_ElfMachineKind_M32            = 1u, 
-  SYMS_ElfMachineKind_SPARC          = 2u, 
-  SYMS_ElfMachineKind_386            = 3u, 
-  SYMS_ElfMachineKind_68K            = 4u, 
-  SYMS_ElfMachineKind_88K            = 5u, 
-  SYMS_ElfMachineKind_IAMCU          = 6u, 
-  SYMS_ElfMachineKind_860            = 7u, 
-  SYMS_ElfMachineKind_MIPS           = 8u, 
-  SYMS_ElfMachineKind_S370           = 9u, 
-  SYMS_ElfMachineKind_MIPS_RS3_LE    = 10u, 
-  SYMS_ElfMachineKind_PARISC         = 15u, 
-  SYMS_ElfMachineKind_PPC_OLD        = 17u, // nick: Old version of PowerPC. Deprecated
-  SYMS_ElfMachineKind_SPARC32PLUS    = 18u, // nick: Sun's "v8plus"
-  SYMS_ElfMachineKind_960            = 19u, 
-  SYMS_ElfMachineKind_PPC            = 20u, 
-  SYMS_ElfMachineKind_PPC64          = 21u, 
-  SYMS_ElfMachineKind_S390           = 22u, 
-  SYMS_ElfMachineKind_SPU            = 23u, 
-  SYMS_ElfMachineKind_V800           = 36u, 
-  SYMS_ElfMachineKind_FR20           = 37u, 
-  SYMS_ElfMachineKind_RH32           = 38u, 
-  SYMS_ElfMachineKind_MCORE          = 39u, 
-  SYMS_ElfMachineKind_ARM            = 40u, 
-  SYMS_ElfMachineKind_SH             = 42u, 
-  SYMS_ElfMachineKind_IA_64          = 50u, 
-  SYMS_ElfMachineKind_X86_64         = 62u, 
-  SYMS_ElfMachineKind_AARCH64        = 183u, 
-  SYMS_ElfMachineKind_RISCV          = 243u, 
-}SYMS_ElfMachineKind;
-
-typedef enum SYMS_ElfClass{
-  SYMS_ElfClass_None = 0,
-  SYMS_ElfClass_32   = 1,
-  SYMS_ElfClass_64   = 2,
-}SYMS_ElfClass;
+////////////////////////////////
+//~ NOTE(nick): Generated
+
+#include "syms/core/generated/syms_meta_elf.h"
+
+////////////////////////////////
+//~ rjf: Elf File Header Types
 
 typedef enum SYMS_ElfIdentifier{
-  SYMS_ElfIdentifier_MAG0    = 0,
-  SYMS_ElfIdentifier_MAG1    = 1,
-  SYMS_ElfIdentifier_MAG2    = 2,
-  SYMS_ElfIdentifier_MAG3    = 3,
-  SYMS_ElfIdentifier_CLASS   = 4,
-  SYMS_ElfIdentifier_DATA    = 5,
-  SYMS_ElfIdentifier_VERSION = 6,
-  SYMS_ElfIdentifier_PAD     = 7,
-  SYMS_ElfIdentifier_NINDENT = 16,
+  SYMS_ElfIdentifier_MAG0       = 0,
+  SYMS_ElfIdentifier_MAG1       = 1,
+  SYMS_ElfIdentifier_MAG2       = 2,
+  SYMS_ElfIdentifier_MAG3       = 3,
+  SYMS_ElfIdentifier_CLASS      = 4,
+  SYMS_ElfIdentifier_DATA       = 5,
+  SYMS_ElfIdentifier_VERSION    = 6,
+  SYMS_ElfIdentifier_OSABI      = 7,
+  SYMS_ElfIdentfiier_ABIVERSION = 8,
+  SYMS_ElfIdentifier_NIDENT     = 16,
 }SYMS_ElfIdentifier;
 
 typedef struct SYMS_ElfEhdr64 SYMS_ElfEhdr64;
 struct SYMS_ElfEhdr64
 {
-  SYMS_U8 e_ident[SYMS_ElfIdentifier_NINDENT];
+  SYMS_U8 e_ident[SYMS_ElfIdentifier_NIDENT];
   SYMS_U16 e_type;
   SYMS_U16 e_machine;
   SYMS_U32 e_version;
@@ -96,7 +47,7 @@ struct SYMS_ElfEhdr64
 typedef struct SYMS_ElfEhdr32 SYMS_ElfEhdr32;
 struct SYMS_ElfEhdr32
 {
-  SYMS_U8 e_ident[SYMS_ElfIdentifier_NINDENT];
+  SYMS_U8 e_ident[SYMS_ElfIdentifier_NIDENT];
   SYMS_U16 e_type;
   SYMS_U16 e_machine;
   SYMS_U32 e_version;
@@ -168,109 +119,11 @@ struct SYMS_ElfPhdr32
   SYMS_U32 p_align;
 };
 
-typedef enum SYMS_ElfSectionCode{
-  SYMS_ElfSectionCode_NULL                   = 0,
-  SYMS_ElfSectionCode_PROGBITS               = 1,
-  SYMS_ElfSectionCode_SYMTAB                 = 2,
-  SYMS_ElfSectionCode_STRTAB                 = 3,
-  SYMS_ElfSectionCode_RELA                   = 4,
-  SYMS_ElfSectionCode_HASH                   = 5,
-  SYMS_ElfSectionCode_DYNAMIC                = 6,
-  SYMS_ElfSectionCode_NOTE                   = 7,
-  SYMS_ElfSectionCode_NOBITS                 = 8,
-  SYMS_ElfSectionCode_REL                    = 9,
-  SYMS_ElfSectionCode_SHLIB                  = 10,
-  SYMS_ElfSectionCode_DYNSYM                 = 11,
-  SYMS_ElfSectionCode_INIT_ARRAY	           = 14,                // nick: Array of ptrs to init functions
-  SYMS_ElfSectionCode_FINI_ARRAY	           = 15,                // nick: Array of ptrs to finish functions
-  SYMS_ElfSectionCode_PREINIT_ARRAY          = 16,                // nick: Array of ptrs to pre-init funcs
-  SYMS_ElfSectionCode_GROUP	                = 17,                // nick: Section contains a section group
-  SYMS_ElfSectionCode_SYMTAB_SHNDX           = 18,                // nick: Indices for SHN_XINDEX entries
-  SYMS_ElfSectionCode_GNU_INCREMENTAL_INPUTS = 0x6fff4700,        // nick: incremental build data
-  SYMS_ElfSectionCode_GNU_ATTRIBUTES         = 0x6ffffff5,        // nick: Object attributes
-  SYMS_ElfSectionCode_GNU_HASH               = 0x6ffffff6,        // nick: GNU style symbol hash table
-  SYMS_ElfSectionCode_GNU_LIBLIST	          = 0x6ffffff7,        // nick: List of prelink dependencies
-  
-  // NOTE(nick): The next three section types are defined by Solaris, and are
-  // named SHT_SUNW*.  We use them in GNU code, so we also define SHT_GNU*
-  // versions.
-  SYMS_ElfSectionCode_SUNW_verdef	   = 0x6ffffffd,	// nick: Versions defined by file
-  SYMS_ElfSectionCode_SUNW_verneed    = 0x6ffffffe,	// nick: Versions needed by file
-  SYMS_ElfSectionCode_SUNW_versym	   = 0x6fffffff,	// nick: Symbol versions
-  
-  SYMS_ElfSectionCode_GNU_verdef	    = SYMS_ElfSectionCode_SUNW_verdef,
-  SYMS_ElfSectionCode_GNU_verneed	   = SYMS_ElfSectionCode_SUNW_verneed,
-  SYMS_ElfSectionCode_GNU_versym	    = SYMS_ElfSectionCode_SUNW_versym,
-  
-  SYMS_ElfSectionCode_LOPROC          = 0x70000000,
-  SYMS_ElfSectionCode_HIPROC          = 0x7fffffff,
-  SYMS_ElfSectionCode_LOUSER          = 0x80000000,
-  SYMS_ElfSectionCode_HIUSER          = 0xffffffff,
-}SYMS_ElfSectionCode;
-
-enum{
-  SYMS_ElfSectionFlag_WRITE            = 0x1,
-  SYMS_ElfSectionFlag_ALLOC            = 0x2,
-  SYMS_ElfSectionFlag_EXECINSTR        = 0x4,
-  SYMS_ElfSectionFlag_MERGE            = 0x10,
-  SYMS_ElfSectionFlag_STRINGS          = 0x20,
-  SYMS_ElfSectionFlag_INFO_LINK        = 0x40,
-  SYMS_ElfSectionFlag_LINK_ORDER       = 0x80,
-  SYMS_ElfSectionFlag_OS_NONCONFORMING = 0x100,
-  SYMS_ElfSectionFlag_GROUP            = 0x200,
-  SYMS_ElfSectionFlag_TLS              = 0x400,
-  SYMS_ElfSectionFlag_MASKOS           = 0xff0000,
-  SYMS_ElfSectionFlag_AMD64_LARGE      = 0x10000000,
-  SYMS_ElfSectionFlag_ORDERED          = 0x40000000,
-  SYMS_ElfSectionFlagT_EXCLUDE         = 0x80000000,
-  SYMS_ElfSectionFlag_MASKPROC         = 0xf0000000
-};
-
 ////////////////////////////////
 //~ NOTE(allen): Auxiliary Vectors
 
 // these appear in /proc/<pid>/auxv of a process, they are not in elf files
 
-typedef enum SYMS_ElfAuxType{
-  SYMS_ElfAuxType_NULL = 0,
-  SYMS_ElfAuxType_PHDR = 3,   // program headers
-  SYMS_ElfAuxType_PHENT = 4,  // size of a program header
-  SYMS_ElfAuxType_PHNUM = 5,  // number of program headers
-  SYMS_ElfAuxType_PAGESZ = 6, // system page size
-  SYMS_ElfAuxType_BASE = 7,   // interpreter base address
-  SYMS_ElfAuxType_FLAGS = 8,
-  SYMS_ElfAuxType_ENTRY = 9,  // program entry point
-  SYMS_ElfAuxType_UID = 11,
-  SYMS_ElfAuxType_EUID = 12,
-  SYMS_ElfAuxType_GID = 13,
-  SYMS_ElfAuxType_EGID = 14,
-  SYMS_ElfAuxType_PLATFORM = 15, // 'platform' as a string (TODO(allen): study)
-  SYMS_ElfAuxType_HWCAP = 16,
-  SYMS_ElfAuxType_CLKTCK = 17,
-  SYMS_ElfAuxType_DCACHEBSIZE = 19,
-  SYMS_ElfAuxType_ICACHEBSIZE = 20,
-  SYMS_ElfAuxType_UCACHEBSIZE = 21,
-  SYMS_ElfAuxType_IGNOREPPC = 22,
-  SYMS_ElfAuxType_SECURE = 23,
-  SYMS_ElfAuxType_BASE_PLATFORM = 24, // 'platform' as a string (different) (TODO(allen): study)
-  SYMS_ElfAuxType_RANDOM = 25, // addres to 16 random bytes
-  SYMS_ElfAuxType_HWCAP2 = 26,
-  SYMS_ElfAuxType_EXECFN = 31, // file name of executable
-  
-  SYMS_ElfAuxType_SYSINFO = 32,
-  SYMS_ElfAuxType_SYSINFO_EHDR = 33,
-  
-  // cool info about caches? (TODO(allen): study)
-  SYMS_ElfAuxType_L1I_CACHESIZE = 40,
-  SYMS_ElfAuxType_L1I_CACHEGEOMETRY = 41,
-  SYMS_ElfAuxType_L1D_CACHESIZE = 42,
-  SYMS_ElfAuxType_L1D_CACHEGEOMETRY = 43,
-  SYMS_ElfAuxType_L2_CACHESIZE = 44,
-  SYMS_ElfAuxType_L2_CACHEGEOMETRY = 45,
-  SYMS_ElfAuxType_L3_CACHESIZE = 46,
-  SYMS_ElfAuxType_L3_CACHEGEOMETRY = 47,
-} SYMS_ElfAuxType;
-
 typedef struct SYMS_ElfAuxv32{
   SYMS_U32 a_type;
   SYMS_U32 a_val;
@@ -286,47 +139,6 @@ typedef struct SYMS_ElfAuxv64{
 
 // these appear in the virtual address space of a process, they are not in elf files
 
-typedef enum SYMS_ElfDynTag{
-  SYMS_ElfDynTag_NULL = 0,
-  SYMS_ElfDynTag_NEEDED = 1,
-  SYMS_ElfDynTag_PLTRELSZ = 2,
-  SYMS_ElfDynTag_PLTGOT = 3,
-  SYMS_ElfDynTag_HASH = 4,
-  SYMS_ElfDynTag_STRTAB = 5,
-  SYMS_ElfDynTag_SYMTAB = 6,
-  SYMS_ElfDynTag_RELA = 7,
-  SYMS_ElfDynTag_RELASZ = 8,
-  SYMS_ElfDynTag_RELAENT = 9,
-  SYMS_ElfDynTag_STRSZ = 10,
-  SYMS_ElfDynTag_SYMENT = 11,
-  SYMS_ElfDynTag_INIT = 12,
-  SYMS_ElfDynTag_FINI = 13,
-  SYMS_ElfDynTag_SONAME = 14,
-  SYMS_ElfDynTag_RPATH = 15,
-  SYMS_ElfDynTag_SYMBOLIC = 16,
-  SYMS_ElfDynTag_REL = 17,
-  SYMS_ElfDynTag_RELSZ = 18,
-  SYMS_ElfDynTag_RELENT = 19,
-  SYMS_ElfDynTag_PLTREL = 20,
-  SYMS_ElfDynTag_DEBUG = 21,
-  SYMS_ElfDynTag_TEXTREL = 22,
-  SYMS_ElfDynTag_JMPREL = 23,
-  SYMS_ElfDynTag_BIND_NOW = 24,
-  SYMS_ElfDynTag_INIT_ARRAY = 25,
-  SYMS_ElfDynTag_FINI_ARRAY = 26,
-  SYMS_ElfDynTag_INIT_ARRAYSZ = 27,
-  SYMS_ElfDynTag_FINI_ARRAYSZ = 28,
-  SYMS_ElfDynTag_RUNPATH = 29,
-  SYMS_ElfDynTag_FLAGS = 30,
-  SYMS_ElfDynTag_PREINIT_ARRAY = 32,
-  SYMS_ElfDynTag_PREINIT_ARRAYSZ = 33,
-  SYMS_ElfDynTag_SYMTAB_SHNDX = 34,
-  SYMS_ElfDynTag_LOOS = 0x6000000D,
-  SYMS_ElfDynTag_HIOS = 0x6ffff000,
-  SYMS_ElfDynTag_LOPROC = 0x70000000,
-  SYMS_ElfDynTag_HIPROC = 0x7fffffff,
-} SYMS_ElfDynTag;
-
 typedef struct SYMS_ElfDyn32{
   SYMS_U32 tag;
   SYMS_U32 val;
@@ -352,19 +164,134 @@ typedef struct SYMS_ElfLinkMap64{
 } SYMS_ElfLinkMap64;
 
 ////////////////////////////////
-//~ rjf: ELF Format Functions
+// Imports and Exports
+
+typedef struct 
+{
+  SYMS_U32 st_name;  // Holds index into files string table.
+  SYMS_U32 st_value; // Depending on the context, this may be address, size, etc.
+  SYMS_U32 st_size;  // Data size in bytes. Zero when size is unknown.
+  SYMS_U8  st_info;  // Contains symbols type and binding.
+  SYMS_U8  st_other; // Reserved for future use, currenly zero.
+  SYMS_U16 st_shndx; // Section index to which symbol is relevant.
+} SYMS_ElfSym32;
+
+typedef struct 
+{
+  SYMS_U32 st_name;
+  SYMS_U8  st_info;
+  SYMS_U8  st_other;
+  SYMS_U16 st_shndx;
+  SYMS_U64 st_value;
+  SYMS_U64 st_size;
+} SYMS_ElfSym64;
+
+#define SYMS_ELF_ST_INFO(b,t)     (((b) << 4) + ((t) & 0xF))
+#define SYMS_ELF_ST_BIND(x)       ((x) >> 4)
+#define SYMS_ELF_ST_TYPE(x)       ((x) & 0xF)
+#define SYMS_ELF_ST_VISIBILITY(v) ((v) & 0x3)
+
+typedef struct
+{
+  SYMS_U32 r_offset;
+  SYMS_U32 r_info;
+} SYMS_ElfRel32;
+
+typedef struct
+{
+  SYMS_U32 r_offset;
+  SYMS_U32 r_info;
+  SYMS_S32 r_addend;
+} SYMS_ElfRela32;
+
+typedef struct
+{
+  SYMS_U64 r_offset;
+  SYMS_U64 r_info;
+} SYMS_ElfRel64;
+
+typedef struct
+{
+  SYMS_U64 r_offset;
+  SYMS_U64 r_info;
+  SYMS_S64 r_addend;
+} SYMS_ElfRela64;
+
+#define SYMS_ELF32_R_SYM(x)  ((x) >> 8)
+#define SYMS_ELF32_R_TYPE(x) ((x) & 0xFF)
+
+#define SYMS_ELF64_R_INFO(s,t) (((SYMS_U64)(s) << 32) | (SYMS_U64)t)
+#define SYMS_ELF64_R_SYM(x)  ((x) >> 32)
+#define SYMS_ELF64_R_TYPE(x) ((x) & 0xffffffff)
+
+// This flag is set to indicate that symbol is not available outside the shared object
+#define SYMS_ELF_EXTERNAL_VERSYM_HIDDEN 0x8000
+#define SYMS_ELF_EXTERNAL_VERSYM_MASK   0x7FFF
+
+// Appears in .gnu.verdef (SHT_GNU_verdef)
+typedef struct
+{
+  SYMS_U16 vd_version;
+  SYMS_U16 vd_flags;
+  SYMS_U16 vd_ndx;
+  SYMS_U16 vd_cnt;
+  SYMS_U32 vd_hash;
+  SYMS_U32 vd_aux;
+  SYMS_U32 vd_next;
+} SYMS_ElfExternalVerdef;
+
+// Appears in .gnu.verdef (SHT_GNU_verdef)
+typedef struct
+{
+  SYMS_U32 vda_name;
+  SYMS_U32 vda_next;
+} SYMS_ElfExternalVerdaux;
+
+// Appears in .gnu.verneed (SHT_GNU_verneed)
+typedef struct
+{
+  SYMS_U16 vn_version;
+  SYMS_U16 vn_cnt;
+  SYMS_U32 vn_file;
+  SYMS_U32 vn_aux;
+  SYMS_U32 vn_next;
+} SYMS_ElfExternalVerneed;
+
+// Appears in .gnu.verneed (SHT_GNU_verneed)
+typedef struct
+{
+  SYMS_U32 vna_hash;
+  SYMS_U16 vna_flags;
+  SYMS_U16 vna_other;
+  SYMS_U32 vna_name;
+  SYMS_U32 vna_next;
+} SYMS_ElfExternalVernaux;
+
+// Appears in .gnu.version (SHT_GNU_versym)
+typedef struct
+{
+  SYMS_U16 vs_vers;
+} SYMS_ElfExternalVersym;
 
 SYMS_C_LINKAGE_BEGIN
 
+////////////////////////////////
+//~ rjf: 32 => 64 bit conversions
+
 // TODO(allen): avoid extra copies (use pointers where convenient)
 SYMS_API SYMS_ElfEhdr64 syms_elf_ehdr64_from_ehdr32(SYMS_ElfEhdr32 h32);
 SYMS_API SYMS_ElfShdr64 syms_elf_shdr64_from_shdr32(SYMS_ElfShdr32 h32);
 SYMS_API SYMS_ElfPhdr64 syms_elf_phdr64_from_phdr32(SYMS_ElfPhdr32 h32);
+SYMS_API SYMS_ElfDyn64  syms_elf_dyn64_from_dyn32(SYMS_ElfDyn32 h32);
+SYMS_API SYMS_ElfSym64  syms_elf_sym64_from_sym32(SYMS_ElfSym32 sym32);
+SYMS_API SYMS_ElfRel64  syms_elf_rel64_from_rel32(SYMS_ElfRel32 rel32);
+SYMS_API SYMS_ElfRela64 syms_elf_rela64_from_rela32(SYMS_ElfRela32 rela32);
 // TODO(allen): auxv?
-// TODO(allen): dyn?
 // TODO(allen): linkmap?
 
-// TODO(allen): promote to base
+////////////////////////////////
+//~ rjf: .gnu_debuglink section 32-bit CRC
+
 SYMS_API SYMS_U32 syms_elf_gnu_debuglink_crc32(SYMS_U32 crc, SYMS_String8 data);
 
 SYMS_C_LINKAGE_END
