Update SymsLib to latest version and update SymsResolver for Insights.

diff --git a/Engine/Source/Runtime/SymsLib/syms/core/elf/syms_elf_parser.c b/Engine/Source/Runtime/SymsLib/syms/core/elf/syms_elf_parser.c
--- a/Engine/Source/Runtime/SymsLib/syms/core/elf/syms_elf_parser.c
+++ b/Engine/Source/Runtime/SymsLib/syms/core/elf/syms_elf_parser.c
@@ -3,9 +3,7 @@
 #define SYMS_ELF_PARSER_C
 
 ////////////////////////////////
-//~ NOTE(rjf): Functions
-
-//- rjf: low-level header/section parsing
+//~ rjf: Low-Level Header/Section Parsing
 
 SYMS_API SYMS_ElfImgHeader
 syms_elf_img_header_from_file(SYMS_String8 file)
@@ -19,7 +17,7 @@ syms_elf_img_header_from_file(SYMS_String8 file)
   SYMS_U64 elf_header_off = 0;
   
   //- rjf: figure out if this file starts with a ELF header
-  SYMS_U8 sig[SYMS_ElfIdentifier_NINDENT];
+  SYMS_U8 sig[SYMS_ElfIdentifier_NIDENT];
   SYMS_B32 sig_is_elf = syms_false;
   {
     syms_memset(sig, 0, sizeof(sig));
@@ -227,6 +225,7 @@ syms_elf_section_array_from_img_header(SYMS_Arena *arena, SYMS_String8 file, SYM
     sections[section_idx].code = (SYMS_ElfSectionCode)header.sh_type;
     sections[section_idx].virtual_range = syms_make_u64_range(header.sh_addr, header.sh_addr + virt_size);
     sections[section_idx].file_range = syms_make_u64_range(header.sh_offset, header.sh_offset + file_size);
+    sections[section_idx].sh_link = header.sh_link;
     sections[section_idx].name = name_stabilized;
   }
   
@@ -276,7 +275,8 @@ syms_elf_ext_debug_ref_from_elf_section_array(SYMS_String8 file, SYMS_ElfSection
   return result;
 }
 
-//- rjf: high-level API canonical conversions
+////////////////////////////////
+//~ rjf: High-Level API Canonical Conversions
 
 SYMS_API SYMS_SecInfo
 syms_elf_section_info_from_elf_section(SYMS_ElfSection elf_section){
@@ -293,7 +293,8 @@ syms_elf_sec_name_from_elf_section(SYMS_ElfSection elf_section){
   return(result);
 }
 
-//- rjf: file accelerator
+////////////////////////////////
+//~ rjf: File Accelerator
 
 SYMS_API SYMS_ElfFileAccel *
 syms_elf_file_accel_from_data(SYMS_Arena *arena, SYMS_String8 string)
@@ -307,7 +308,8 @@ syms_elf_file_accel_from_data(SYMS_Arena *arena, SYMS_String8 string)
   return file_accel;
 }
 
-//- rjf: binary
+////////////////////////////////
+//~ rjf: Binary
 
 SYMS_API SYMS_ElfBinAccel *
 syms_elf_bin_accel_from_file(SYMS_Arena *arena, SYMS_String8 data, SYMS_ElfFileAccel *file_accel)
@@ -365,7 +367,7 @@ syms_elf_arch_from_bin(SYMS_ElfBinAccel *bin){
 }
 
 ////////////////////////////////
-// NOTE(allen): ELF Specific Helpers
+//~ NOTE(allen): ELF Specific Helpers
 
 SYMS_API SYMS_ElfSection*
 syms_elf_sec_from_bin_name__unstable(SYMS_ElfBinAccel *bin, SYMS_String8 name){
@@ -381,4 +383,319 @@ syms_elf_sec_from_bin_name__unstable(SYMS_ElfBinAccel *bin, SYMS_String8 name){
   return(result);
 }
 
+////////////////////////////////
+//~ rjf: Imports/Exports
+
+SYMS_API SYMS_ImportArray
+syms_elf_imports_from_bin(SYMS_Arena *arena, SYMS_String8 data, SYMS_ElfBinAccel *bin)
+{
+  SYMS_ArenaTemp scratch = syms_get_scratch(&arena, 1);
+  
+  //- rjf: grab prerequisites
+  void *base = (void *)data.str;
+  SYMS_B32 bin_is_32bit = bin->header.is_32bit;
+  SYMS_U64 sym_size = bin_is_32bit ? sizeof(SYMS_ElfSym32) : sizeof(SYMS_ElfSym64);
+  
+  //- rjf: find ranges
+  SYMS_U64Range verneed_range        = syms_make_u64_range(0,0);
+  SYMS_U64Range verneed_strtab_range = syms_make_u64_range(0,0);
+  SYMS_U64Range versym_range         = syms_make_u64_range(0,0);
+  SYMS_U64Range dynsym_range         = syms_make_u64_range(0,0);
+  SYMS_U64Range dynsym_strtab_range  = syms_make_u64_range(0,0);
+  {
+    SYMS_ElfSection *first = &bin->sections.v[0];
+    SYMS_ElfSection *opl = first+bin->sections.count;
+    for(SYMS_ElfSection *sect = first; sect < opl; sect += 1) 
+    {
+      // NOTE(rjf): We subtract 1 from sh_link because we don't store the null
+      // section in the bin section array.
+      SYMS_U64 sh_link_idx = sect->sh_link-1;
+      SYMS_B32 sh_link_in_bounds = (0 < sect->sh_link && sect->sh_link < bin->sections.count);
+      if(syms_string_match(sect->name, syms_str8_lit(".gnu.version_r"), 0)) 
+      {
+        verneed_range = sect->file_range;
+        verneed_strtab_range = bin->sections.v[sh_link_idx].file_range;
+      }
+      if(syms_string_match(sect->name, syms_str8_lit(".gnu.version"), 0)) 
+      {
+        versym_range = sect->file_range;
+      }
+      if(syms_string_match(sect->name, syms_str8_lit(".dynsym"), 0)) 
+      {
+        dynsym_range = sect->file_range;
+        dynsym_strtab_range = bin->sections.v[sh_link_idx].file_range;
+      }
+    }
+  }
+  
+  //- rjf: grab .dynsym symbol array
+  SYMS_U64 sym_arr_count = syms_u64_range_size(dynsym_range) / sym_size;
+  SYMS_ElfSym64 *sym_arr = syms_push_array_zero(scratch.arena, SYMS_ElfSym64, sym_arr_count);
+  SYMS_ElfSym64 *sym_arr_opl = sym_arr + sym_arr_count;
+  {
+    // NOTE(rjf): We start the read offset at sym_size because the format
+    // always has a meaningless null symbol at the beginning of the block.
+    SYMS_U64 roff__sym_array = sym_size;
+    if(bin_is_32bit)
+    {
+      for(SYMS_U64 idx = 0; idx < sym_arr_count; idx += 1)
+      {
+        SYMS_ElfSym32 *sym32 = (SYMS_ElfSym32 *)syms_based_range_ptr(base, dynsym_range, roff__sym_array + idx*sym_size);
+        sym_arr[idx] = syms_elf_sym64_from_sym32(*sym32);
+      }
+    }
+    else
+    {
+      syms_memmove(sym_arr, syms_based_range_ptr(base, dynsym_range, roff__sym_array), sym_arr_count*sym_size);
+    }
+  }
+  
+  //- rjf: parse all import info
+  SYMS_ImportNode *first_import = 0;
+  SYMS_ImportNode *last_import = 0;
+  SYMS_U64 import_count = 0;
+  {
+    SYMS_U64 sym_idx = 1;
+    for(SYMS_ElfSym64 *sym = sym_arr; sym < sym_arr_opl; sym += 1, sym_idx += 1)
+    {
+      // rjf: skip any symbol that does not have the right section index
+      if(sym->st_shndx != SYMS_ElfSectionIndex_UNDEF)
+      {
+        continue;
+      }
+      
+      // rjf: read versym
+      SYMS_U64 roff__versym = sym_idx * sizeof(SYMS_ElfExternalVersym);
+      SYMS_ElfExternalVersym vs = {0};
+      syms_based_range_read_struct(base, versym_range, roff__versym, &vs);
+      
+      // rjf: parse vn_file / vna_name
+      SYMS_String8 vn_file = syms_str8_lit("");
+      SYMS_String8 vna_name = syms_str8_lit("");
+      {
+        // rjf: parse all verneeds
+        SYMS_U64 verneed_read_offset = 0;
+        SYMS_U16 version = vs.vs_vers & SYMS_ELF_EXTERNAL_VERSYM_MASK;
+        SYMS_ElfExternalVerneed vn = {0};
+        for(;;)
+        {
+          syms_based_range_read_struct(base, verneed_range, verneed_read_offset, &vn);
+          
+          // rjf: find vernaux with a matching version
+          SYMS_U64 aux_read_offset = verneed_read_offset + vn.vn_aux;
+          for(SYMS_U32 aux_idx = 0; aux_idx < vn.vn_cnt; aux_idx += 1) 
+          {
+            // rjf: read
+            SYMS_ElfExternalVernaux vna = {0};
+            syms_based_range_read_struct(base, verneed_range, aux_read_offset, &vna);
+            
+            // rjf: if the version matches, we've found the right strings
+            if(vna.vna_other == version) 
+            {
+              vn_file  = syms_based_range_read_string(base, verneed_strtab_range, vn.vn_file);
+              vna_name = syms_based_range_read_string(base, verneed_strtab_range, vna.vna_name);
+              goto end__vn_file__vna_name__parse;
+            }
+            
+            // rjf: advance
+            if(vna.vna_next == 0) 
+            {
+              break;
+            }
+            aux_read_offset += vna.vna_next;
+          }
+          
+          // rjf: advance
+          if(vn.vn_next == 0) 
+          {
+            break;
+          }
+          verneed_read_offset += vn.vn_next;
+        }
+      }
+      end__vn_file__vna_name__parse:;
+      
+      // rjf: build node
+      SYMS_ImportNode *node = syms_push_array_zero(scratch.arena, SYMS_ImportNode, 1);
+      SYMS_QueuePush(first_import, last_import, node);
+      import_count += 1;
+      {
+        SYMS_Import *imp  = &node->data;
+        
+        // rjf: fill out symbol name
+        {
+          SYMS_ArenaTemp temp = syms_arena_temp_begin(scratch.arena);
+          SYMS_String8List list = {0};
+          SYMS_String8 symbol_name = syms_based_range_read_string(base, dynsym_strtab_range, sym->st_name);
+          syms_string_list_push(temp.arena, &list, symbol_name);
+          if(vna_name.size > 0) 
+          {
+            SYMS_B32 is_hidden = !!(vs.vs_vers & SYMS_ELF_EXTERNAL_VERSYM_HIDDEN);
+            syms_string_list_push(temp.arena, &list, is_hidden ? syms_str8_lit("@@") : syms_str8_lit("@"));
+            syms_string_list_push(temp.arena, &list, vna_name);
+          }
+          imp->name = syms_string_list_join(arena, &list, 0);
+          syms_arena_temp_end(temp);
+        }
+        
+        // rjf: fill library name
+        {
+          imp->library_name = syms_push_string_copy(arena, vn_file);
+        }
+      }
+    }
+  }
+  
+  //- rjf: build/fill result
+  SYMS_ImportArray result = {0};
+  {
+    result.count = import_count;
+    result.imports = syms_push_array_zero(arena, SYMS_Import, result.count);
+    SYMS_U64 idx = 0;
+    for(SYMS_ImportNode *in = first_import;
+        in != 0;
+        in = in->next, idx += 1)
+    {
+      result.imports[idx] = in->data;
+    }
+  }
+  
+  //- rjf: return
+  syms_release_scratch(scratch);
+  return result;
+}
+
+SYMS_API SYMS_ExportArray
+syms_elf_exports_from_bin(SYMS_Arena *arena, SYMS_String8 data, SYMS_ElfBinAccel *bin)
+{
+  SYMS_ArenaTemp scratch = syms_get_scratch(&arena, 1);
+  
+  SYMS_U64Range verdef_range         = syms_make_u64_range(0,0);
+  SYMS_U64Range verdef_strtab_range  = syms_make_u64_range(0,0);
+  SYMS_U64Range versym_range         = syms_make_u64_range(0,0);
+  SYMS_U64Range dynsym_range         = syms_make_u64_range(0,0);
+  SYMS_U64Range dynsym_strtab_range  = syms_make_u64_range(0,0);
+  for (SYMS_U32 sec_idx = 0; sec_idx < bin->sections.count; sec_idx += 1) 
+  {
+    SYMS_ElfSection *sect = &bin->sections.v[sec_idx];
+    if (syms_string_match(sect->name, syms_str8_lit(".gnu.version_d"), 0)) 
+    {
+      SYMS_ASSERT(sect->sh_link > 0 && sect->sh_link < bin->sections.count);
+      verdef_range = sect->file_range;
+      verdef_strtab_range = bin->sections.v[sect->sh_link-1].file_range;
+    }
+    if (syms_string_match(sect->name, syms_str8_lit(".gnu.version"), 0)) 
+    {
+      versym_range = sect->file_range;
+    }
+    if (syms_string_match(sect->name, syms_str8_lit(".dynsym"), 0)) 
+    {
+      SYMS_ASSERT(sect->sh_link > 0 && sect->sh_link < bin->sections.count);
+      dynsym_range = sect->file_range;
+      dynsym_strtab_range = bin->sections.v[sect->sh_link-1].file_range;
+    }
+  }
+  
+  SYMS_ExportNode *first_export = 0;
+  SYMS_ExportNode *last_export = 0;
+  SYMS_U32 export_count = 0;
+  void *base = (void*)data.str;
+  SYMS_U64 read_offset = 0;
+  for (; read_offset < syms_u64_range_size(dynsym_range);)
+  {
+    SYMS_U32 sym_idx = read_offset / (bin->header.is_32bit ? sizeof(SYMS_ElfSym32) : sizeof(SYMS_ElfSym64));
+    SYMS_ElfSym64 sym;
+    if (bin->header.is_32bit) 
+    {
+      SYMS_ElfSym32 sym32; syms_memzero_struct(&sym32);
+      read_offset += syms_based_range_read_struct(base, dynsym_range, read_offset, &sym32);
+      sym = syms_elf_sym64_from_sym32(sym32);
+    }
+    else
+    {
+      syms_memzero_struct(&sym);
+      read_offset += syms_based_range_read_struct(base, dynsym_range, read_offset, &sym);
+    }
+    SYMS_U32 st_bind = SYMS_ELF_ST_BIND(sym.st_info);
+    SYMS_B32 is_symbol_visible_to_other_objects = st_bind == SYMS_ElfSymBind_GLOBAL || st_bind == SYMS_ElfSymBind_WEAK;
+    SYMS_B32 is_symbol_defined = sym.st_shndx != SYMS_ElfSectionIndex_UNDEF;
+    if (is_symbol_visible_to_other_objects && is_symbol_defined) 
+    {
+      SYMS_StringJoin join; syms_memzero_struct(&join);
+      SYMS_String8List name_list; syms_memzero_struct(&name_list);
+      
+      SYMS_String8 st_name = syms_based_range_read_string(base, dynsym_strtab_range, sym.st_name);
+      syms_string_list_push(scratch.arena, &name_list, st_name);
+      
+      {
+        // find symbol version token
+        SYMS_ElfExternalVersym vs;
+        syms_based_range_read_struct(base, versym_range, sym_idx * sizeof(vs), &vs);
+        SYMS_U32 symbol_version = vs.vs_vers & SYMS_ELF_EXTERNAL_VERSYM_MASK;
+        // iterate defined version to find version string
+        SYMS_U64 verdef_read_offset = 0;
+        for (;verdef_read_offset < syms_u64_range_size(verdef_range);) 
+        {
+          SYMS_ElfExternalVerdef vd; syms_memzero_struct(&vd);
+          syms_based_range_read_struct(base, verdef_range, verdef_read_offset, &vd);
+          SYMS_U32 version = vd.vd_ndx & SYMS_ELF_EXTERNAL_VERSYM_MASK;
+          if (version == symbol_version) 
+          {
+            SYMS_B32 has_name = vd.vd_ndx > 1 && (~vd.vd_flags & SYMS_ElfExternalVerFlag_BASE);
+            if (has_name) 
+            {
+              SYMS_U64 aux_read_offset = verdef_read_offset + vd.vd_aux;
+              SYMS_ElfExternalVerdaux vda; syms_memzero_struct(&vda);
+              syms_based_range_read_struct(base, verdef_range, aux_read_offset, &vda);
+              SYMS_String8 verdef_name = syms_based_range_read_string(base, verdef_strtab_range, vda.vda_name);
+              SYMS_B32 is_hidden = !!(vs.vs_vers & SYMS_ELF_EXTERNAL_VERSYM_HIDDEN);
+              syms_string_list_push(scratch.arena, &name_list, is_hidden ? syms_str8_lit("@@") : syms_str8_lit("@"));
+              syms_string_list_push(scratch.arena, &name_list, verdef_name);
+            } 
+            break;
+          }
+          if (vd.vd_next == 0) 
+          {
+            break;
+          }
+          verdef_read_offset += vd.vd_next;
+        }
+      }
+      
+      SYMS_U64 symbol_base_address = 0;
+      if (sym.st_shndx < SYMS_ElfSectionIndex_LO_RESERVE) 
+      {
+        SYMS_ASSERT(sym.st_shndx > 0 && sym.st_shndx <= bin->sections.count);
+        SYMS_ElfSection *sect = &bin->sections.v[sym.st_shndx-1];
+        symbol_base_address = sect->virtual_range.min;
+      }
+      
+      // add to list
+      SYMS_ExportNode *node = syms_push_array_zero(scratch.arena, SYMS_ExportNode, 1);
+      SYMS_QueuePush(first_export, last_export, node);
+      export_count += 1;
+      
+      // fill out data
+      SYMS_Export *exp = &node->data;
+      exp->name = syms_string_list_join(arena, &name_list, &join);
+      exp->address = symbol_base_address + sym.st_value;
+      exp->ordinal = 0;
+      exp->forwarder_library_name = syms_str8(0,0);
+      exp->forwarder_import_name = syms_str8(0,0);
+    }
+  }
+  
+  SYMS_ExportArray export_array;
+  export_array.count = 0;
+  export_array.exports = syms_push_array(arena, SYMS_Export, export_count);
+  for (SYMS_ExportNode *en = first_export; en != 0; en = en->next) 
+  {
+    export_array.exports[export_array.count++] = en->data;
+  }
+  
+  syms_release_scratch(scratch);
+  
+  return export_array;
+}
+
 #endif // SYMS_ELF_PARSER_C
diff --git a/Engine/Source/Runtime/SymsLib/syms/core/elf/syms_elf_parser.h b/Engine/Source/Runtime/SymsLib/syms/core/elf/syms_elf_parser.h
--- a/Engine/Source/Runtime/SymsLib/syms/core/elf/syms_elf_parser.h
+++ b/Engine/Source/Runtime/SymsLib/syms/core/elf/syms_elf_parser.h
@@ -26,6 +26,7 @@ struct SYMS_ElfSection
   SYMS_U64Range virtual_range;
   SYMS_U64Range file_range;
   SYMS_String8 name;
+  SYMS_U32 sh_link;
 };
 
 typedef struct SYMS_ElfSectionArray SYMS_ElfSectionArray;
@@ -63,23 +64,28 @@ struct SYMS_ElfBinAccel
   SYMS_ElfSectionArray sections;
 };
 
-////////////////////////////////
-//~ NOTE(rjf): Functions
-
 SYMS_C_LINKAGE_BEGIN
 
-//- rjf: low-level header/section parsing
+////////////////////////////////
+//~ rjf: Low-Level Header/Section Parsing
+
 SYMS_API SYMS_ElfImgHeader syms_elf_img_header_from_file(SYMS_String8 file);
 SYMS_API SYMS_ElfSectionArray syms_elf_section_array_from_img_header(SYMS_Arena *arena, SYMS_String8 file, SYMS_ElfImgHeader img);
 SYMS_API SYMS_ElfExtDebugRef syms_elf_ext_debug_ref_from_elf_section_array(SYMS_String8 file, SYMS_ElfSectionArray sections);
 
-//- rjf: high-level API canonical conversions
+////////////////////////////////
+//~ rjf: High-Level API Canonical Conversions
+
 SYMS_API SYMS_SecInfo syms_elf_section_info_from_elf_section(SYMS_ElfSection elf_section);
 
-//- rjf: file accelerator
+////////////////////////////////
+//~ rjf: File Accelerator
+
 SYMS_API SYMS_ElfFileAccel *syms_elf_file_accel_from_data(SYMS_Arena *arena, SYMS_String8 string);
 
-//- rjf: binary
+////////////////////////////////
+//~ rjf: Binary
+
 SYMS_API SYMS_ElfBinAccel *syms_elf_bin_accel_from_file(SYMS_Arena *arena, SYMS_String8 data, SYMS_ElfFileAccel *file_accel);
 SYMS_API SYMS_ExtFileList  syms_elf_ext_file_list_from_bin(SYMS_Arena *arena, SYMS_String8 file, SYMS_ElfBinAccel *bin_accel);
 SYMS_API SYMS_SecInfoArray syms_elf_sec_info_array_from_bin(SYMS_Arena *arena, SYMS_String8 data, SYMS_ElfBinAccel *bin);
@@ -88,10 +94,16 @@ SYMS_API SYMS_U64          syms_elf_default_vbase_from_bin(SYMS_ElfBinAccel *bin
 SYMS_API SYMS_Arch         syms_elf_arch_from_bin(SYMS_ElfBinAccel *bin);
 
 ////////////////////////////////
-// NOTE(allen): ELF Specific Helpers
+//~ NOTE(allen): ELF Specific Helpers
 
 SYMS_API SYMS_ElfSection*   syms_elf_sec_from_bin_name__unstable(SYMS_ElfBinAccel *bin, SYMS_String8 name);
 
+////////////////////////////////
+//~ rjf: Imports/Exports
+
+SYMS_API SYMS_ImportArray syms_elf_imports_from_bin(SYMS_Arena *arena, SYMS_String8 data, SYMS_ElfBinAccel *bin);
+SYMS_API SYMS_ExportArray syms_elf_exports_from_bin(SYMS_Arena *arena, SYMS_String8 data, SYMS_ElfBinAccel *bin);
+
 SYMS_C_LINKAGE_END
 
 #endif // SYMS_ELF_PARSER_H
diff --git a/Engine/Source/Runtime/SymsLib/syms/core/group/syms_group.c b/Engine/Source/Runtime/SymsLib/syms/core/group/syms_group.c
--- a/Engine/Source/Runtime/SymsLib/syms/core/group/syms_group.c
+++ b/Engine/Source/Runtime/SymsLib/syms/core/group/syms_group.c
@@ -24,7 +24,7 @@ syms_group_release(SYMS_Group *group){
 SYMS_API void
 syms_group_init(SYMS_Group *group, SYMS_ParseBundle *params){
   SYMS_ProfBegin("syms_group_init");
-  SYMS_ASSERT(syms_thread_lane == 0);
+  SYMS_ASSERT_PARANOID(syms_thread_lane == 0);
   
   //- fill in bin
   SYMS_String8 bin_data = params->bin_data;
@@ -166,7 +166,7 @@ syms_group_parse_all_top_level(SYMS_Group *group){
 
 SYMS_API void
 syms_group_begin_multilane(SYMS_Group *group, SYMS_U64 lane_count){
-  SYMS_ASSERT(syms_thread_lane == 0);
+  SYMS_ASSERT_PARANOID(syms_thread_lane == 0);
   SYMS_ASSERT(group->lane_count == 1);
   if (lane_count > 1){
     // allocate the group's lane arenas array
@@ -455,7 +455,7 @@ syms_group_file_table_from_uid_with_fallbacks(SYMS_Group *group, SYMS_UnitID uid
 SYMS_API SYMS_UnitRangeArray
 syms_group_unit_ranges(SYMS_Group *group){
   if (!group->unit_ranges_is_filled){
-    SYMS_ASSERT(syms_thread_lane == 0);
+    SYMS_ASSERT_PARANOID(syms_thread_lane == 0);
     group->unit_ranges = syms_unit_ranges_from_set(group->arena, group->dbg_data, group->dbg, group->unit_set);
     group->unit_ranges_is_filled = syms_true;
   }
@@ -490,7 +490,7 @@ syms_group_default_vbase(SYMS_Group *group){
 SYMS_API SYMS_String8
 syms_group_file_name_from_id(SYMS_Group *group, SYMS_UnitID uid, SYMS_FileID file_id){
   SYMS_ProfBegin("syms_group_file_name_from_id");
-  SYMS_ASSERT(syms_thread_lane == 0);
+  SYMS_ASSERT_PARANOID(syms_thread_lane == 0);
   SYMS_String8 from_map = syms_file_id_2_name_map_name_from_id(&group->file_id_2_name_map, uid, file_id);
   SYMS_String8 result = {0};
   if (from_map.size != 0){
@@ -524,6 +524,18 @@ syms_group_type_map(SYMS_Group *group){
   return(result);
 }
 
+SYMS_API SYMS_StrippedInfoArray
+syms_group_stripped_info(SYMS_Group *group){
+  if (!group->stripped_info_is_filled){
+    SYMS_UnitID pub_uid = syms_uid_collated_public_symbols_from_set(group->unit_set);
+    SYMS_UnitAccel *pub = syms_group_unit_from_uid(group, pub_uid);
+    group->stripped_info = syms_stripped_from_unit(group->arena, group->dbg_data, group->dbg, pub);
+    group->stripped_info_is_filled = syms_true;
+  }
+  SYMS_StrippedInfoArray result = group->stripped_info;
+  return(result);
+}
+
 
 ////////////////////////////////
 // NOTE(allen): Syms Group Address Mapping Functions
@@ -744,11 +756,12 @@ syms_group_vranges_from_uid_line__linear_scan(SYMS_Arena *arena, SYMS_Group *gro
   return(result);
 }
 
+//- map getters
 SYMS_API SYMS_SpatialMap1D*
 syms_group_sec_map_v(SYMS_Group *group){
   SYMS_ProfBegin("syms_group_sec_map_v");
   if (!group->sec_map_v_is_built){
-    SYMS_ASSERT(syms_thread_lane == 0);
+    SYMS_ASSERT_PARANOID(syms_thread_lane == 0);
     
     //- fill spatial map array
     SYMS_SecInfoArray sec_array = syms_group_sec_info_array(group);
@@ -790,7 +803,7 @@ SYMS_API SYMS_SpatialMap1D*
 syms_group_sec_map_f(SYMS_Group *group){
   SYMS_ProfBegin("syms_group_sec_map_f");
   if (!group->sec_map_f_is_built){
-    SYMS_ASSERT(syms_thread_lane == 0);
+    SYMS_ASSERT_PARANOID(syms_thread_lane == 0);
     
     //- fill spatial map array
     SYMS_SecInfoArray sec_array = syms_group_sec_info_array(group);
@@ -832,7 +845,7 @@ SYMS_API SYMS_SpatialMap1D*
 syms_group_unit_map(SYMS_Group *group){
   SYMS_ProfBegin("syms_group_unit_map");
   if (!group->unit_map_is_built){
-    SYMS_ASSERT(syms_thread_lane == 0);
+    SYMS_ASSERT_PARANOID(syms_thread_lane == 0);
     
     //- fill spatial map array
     SYMS_UnitRangeArray unit_ranges = syms_group_unit_ranges(group);
@@ -929,7 +942,7 @@ syms_group_line_sequence_map_from_uid(SYMS_Group *group, SYMS_UnitID uid){
     SYMS_Arena *arena = syms_group_get_lane_arena(group);
     SYMS_LineTable *line_table = syms_group_line_table_from_uid(group, uid);
     SYMS_U64 count = line_table->sequence_count;
-    SYMS_SpatialMap1DRange *ranges = syms_push_array_zero(arena, SYMS_SpatialMap1DRange, count);
+    SYMS_SpatialMap1DRange *ranges = syms_push_array(arena, SYMS_SpatialMap1DRange, count);
     SYMS_SpatialMap1DRange *range_ptr = ranges;
     SYMS_U64 *seq_index_ptr = line_table->sequence_index_array;
     SYMS_Line *line_array = line_table->line_array;
@@ -977,7 +990,7 @@ syms_group_line_sequence_map_from_uid(SYMS_Group *group, SYMS_UnitID uid){
 SYMS_API SYMS_SpatialMap1D*
 syms_group_var_map_from_uid(SYMS_Group *group, SYMS_UnitID uid){
   SYMS_ProfBegin("syms_group_var_map_from_uid");
-  SYMS_ASSERT(syms_thread_lane == 0);
+  SYMS_ASSERT_PARANOID(syms_thread_lane == 0);
   SYMS_U64 index = uid - 1;
   if (index < group->unit_count &&
       (group->unit_cache_flags[index] & SYMS_GroupUnitCacheFlag_HasVarMap) == 0){
@@ -1228,6 +1241,54 @@ syms_group_line_to_addr_map_from_uid_file_id(SYMS_Group *group, SYMS_UnitID uid,
   return(result);
 }
 
+SYMS_API SYMS_SpatialMap1D*
+syms_group_stripped_info_map(SYMS_Group *group){
+  SYMS_ProfBegin("syms_group_stripped_info_map");
+  SYMS_ASSERT_PARANOID(syms_thread_lane == 0);
+  if (!group->stripped_info_map_is_built){
+    SYMS_Arena *arena = group->arena;
+    
+    //- fill spatial map array
+    SYMS_StrippedInfoArray stripped_info = syms_group_stripped_info(group);
+    SYMS_U64 guess_count = stripped_info.count;
+    SYMS_SpatialMap1DRange *ranges = syms_push_array(arena, SYMS_SpatialMap1DRange, guess_count);
+    SYMS_SpatialMap1DRange *range_ptr = ranges;
+    SYMS_StrippedInfo *stripped_ptr = stripped_info.info;
+    for (SYMS_U64 i = 0; i < guess_count; i += 1, stripped_ptr += 1){
+      SYMS_U64 min = stripped_ptr->voff;
+      SYMS_U64 max = SYMS_U64_MAX;
+      if (i + 1 < guess_count){
+        max = (stripped_ptr + 1)->voff;
+      }
+      SYMS_ASSERT(min <= max);
+      if (min < max){
+        range_ptr->range.min = min;
+        range_ptr->range.max = max;
+        range_ptr->val = i + 1;
+        range_ptr += 1;
+      }
+    }
+    
+    //- shrink the allocated array
+    SYMS_U64 actual_count = (SYMS_U64)(range_ptr - ranges);
+    syms_arena_put_back(arena, sizeof(SYMS_SpatialMap1D)*(guess_count - actual_count));
+    
+    //- assemble map
+    SYMS_SpatialMap1D map = {ranges, actual_count};
+    
+    //- save to group
+    group->stripped_info_map_is_built = syms_true;
+    group->stripped_info_map = map;
+    
+    SYMS_ASSERT_PARANOID(syms_spatial_map_1d_invariants(&map));
+  }
+  
+  SYMS_SpatialMap1D *result = &group->stripped_info_map;
+  SYMS_ProfEnd();
+  return(result);
+}
+
+//- accelerated versions
 SYMS_API SYMS_U64
 syms_group_sec_number_from_voff__accelerated(SYMS_Group *group, SYMS_U64 voff){
   SYMS_ProfBegin("syms_group_sec_number_from_voff__accelerated");
@@ -1454,7 +1515,7 @@ SYMS_GLOBAL SYMS_TypeNode syms_type_graph_nil = {
 SYMS_API SYMS_TypeNode*
 syms_group_type_from_usid(SYMS_Group *group, SYMS_USID usid){
   SYMS_ProfBegin("syms_group_type_from_usid");
-  SYMS_ASSERT(syms_thread_lane == 0);
+  SYMS_ASSERT_PARANOID(syms_thread_lane == 0);
   SYMS_TypeNode *result = syms_group_type_from_usid__rec(group, usid);
   SYMS_ProfEnd();
   return(result);
@@ -1543,98 +1604,6 @@ syms_group_type_from_usid__rec(SYMS_Group *group, SYMS_USID usid){
             // resolve direct type
             SYMS_TypeNode *direct = syms_group_type_from_usid__rec(group, type_info.direct_type);
             
-            // read members
-            SYMS_U64 mem_count = 0;
-            SYMS_TypeMember *mems = 0;
-            if (type_info.kind == SYMS_TypeKind_Class ||
-                type_info.kind == SYMS_TypeKind_Struct ||
-                type_info.kind == SYMS_TypeKind_Union){
-              SYMS_MemsAccel *mems_accel = syms_mems_accel_from_sid(scratch.arena, dbg_data, dbg, unit,
-                                                                    usid.sid);
-              mem_count = syms_mem_count_from_mems(mems_accel);
-              mems = syms_push_array(group->arena, SYMS_TypeMember, mem_count);
-              
-              SYMS_TypeMember *mem_ptr = mems;
-              for (SYMS_U64 n = 1; n <= mem_count; n += 1, mem_ptr += 1){
-                SYMS_ArenaTemp temp = syms_arena_temp_begin(scratch.arena);
-                
-                SYMS_MemInfo mem_info = syms_mem_info_from_number(group->arena, dbg_data, dbg, unit, mems_accel, n);
-                
-                SYMS_TypeNode *type = &syms_type_graph_nil;
-                switch (mem_info.kind){
-                  case SYMS_MemKind_VTablePtr: /*no type*/ break;
-                  
-                  case SYMS_MemKind_DataField:
-                  case SYMS_MemKind_StaticData:
-                  {
-                    // TODO(allen): here we lose access to the fast path from StaticData to a virtual offset.
-                    // This path is only available from the DWARF backend, but there are a few ways we could
-                    // preserve it.
-                    SYMS_USID type_usid = syms_type_from_mem_number(dbg_data, dbg, unit, mems_accel, n);
-                    type = syms_group_type_from_usid__rec(group, type_usid);
-                  }break;
-                  
-                  case SYMS_MemKind_BaseClass:
-                  case SYMS_MemKind_VBaseClassPtr:
-                  case SYMS_MemKind_NestedType:
-                  {
-                    // directly get type from member id
-                    SYMS_USID type_usid = syms_type_from_mem_number(dbg_data, dbg, unit, mems_accel, n);
-                    type = syms_group_type_from_usid__rec(group, type_usid);
-                  }break;
-                  
-                  case SYMS_MemKind_Method:
-                  case SYMS_MemKind_StaticMethod:
-                  {
-                    // TODO(allen): here we lose access to the fast path a procedure symbol.
-                    // This path is only available from the DWARF backend, but there are a few ways we could
-                    // preserve it.
-                    SYMS_SigInfo sig = syms_sig_info_from_mem_number(scratch.arena, dbg_data, dbg, unit, mems_accel, n);
-                    
-                    SYMS_USID return_usid = syms_make_usid(sig.uid, sig.return_type_id);
-                    SYMS_USID this_usid = syms_make_usid(sig.uid, sig.this_type_id);
-                    
-                    SYMS_TypeNode *ret_type  = syms_group_type_from_usid__rec(group, return_usid);
-                    SYMS_TypeNode *this_type = syms_group_type_from_usid__rec(group, this_usid);
-                    SYMS_U64 param_count = sig.param_type_ids.count;
-                    SYMS_TypeNode **params = syms_push_array(group->arena, SYMS_TypeNode*, param_count);
-                    
-                    {
-                      SYMS_TypeNode **param_ptr = params;
-                      SYMS_TypeNode **param_opl = params + sig.param_type_ids.count;
-                      SYMS_SymbolID *param_id_ptr = sig.param_type_ids.ids;
-                      for (; param_ptr < param_opl; param_id_ptr += 1, param_ptr += 1){
-                        *param_ptr = syms_group_type_from_usid__rec(group, syms_make_usid(sig.uid, *param_id_ptr));
-                      }
-                    }
-                    
-                    type = syms_group_type_proc_from_type(group, ret_type, this_type, params, param_count);
-                  }break;
-                }
-                
-                mem_ptr->kind = mem_info.kind;
-                mem_ptr->visibility = mem_info.visibility;
-                mem_ptr->flags = mem_info.flags;
-                mem_ptr->name = syms_string_cons(group->arena, &group->string_cons, mem_info.name);
-                mem_ptr->off = mem_info.off;
-                mem_ptr->virtual_off = mem_info.virtual_off;
-                mem_ptr->type = type;
-                
-                syms_arena_temp_end(temp);
-              }
-            }
-            
-            // read constants
-            SYMS_EnumInfoArray enum_info_array = {0};
-            if (type_info.kind == SYMS_TypeKind_Enum){
-              enum_info_array = syms_enum_info_array_from_sid(group->arena, dbg_data, dbg, unit, usid.sid);
-              
-              SYMS_EnumInfo *enum_info = enum_info_array.enum_info;
-              for (SYMS_U64 i = 0; i < enum_info_array.count; i += 1, enum_info += 1){
-                enum_info->name = syms_string_cons(group->arena, &group->string_cons, enum_info->name);
-              }
-            }
-            
             // compute size info
             SYMS_U64 byte_size = 0;
             switch (type_info.reported_size_interp){
@@ -1916,7 +1885,7 @@ syms_group_enum_members_from_type(SYMS_Group *group, SYMS_TypeNode *type){
 SYMS_API SYMS_U64
 syms_group_type_size_from_usid(SYMS_Group *group, SYMS_USID usid){
   SYMS_ProfBegin("syms_group_type_size_from_usid");
-  SYMS_ASSERT(syms_thread_lane == 0);
+  SYMS_ASSERT_PARANOID(syms_thread_lane == 0);
   SYMS_UnitAccel *unit = syms_group_unit_from_uid(group, usid.uid);
   SYMS_TypeInfo type_info = syms_type_info_from_sid(group->dbg_data, group->dbg, unit, usid.sid);
   SYMS_U64 result = 0;
@@ -1943,7 +1912,7 @@ syms_group_type_size_from_usid(SYMS_Group *group, SYMS_USID usid){
 SYMS_API SYMS_TypeNode*
 syms_group_type_basic(SYMS_Group *group, SYMS_TypeKind basic_kind, SYMS_U64 size, SYMS_String8 name){
   SYMS_ProfBegin("syms_group_type_basic");
-  SYMS_ASSERT(syms_thread_lane == 0);
+  SYMS_ASSERT_PARANOID(syms_thread_lane == 0);
   
   SYMS_TypeNode *result = &syms_type_graph_nil;
   
@@ -1974,7 +1943,7 @@ syms_group_type_basic(SYMS_Group *group, SYMS_TypeKind basic_kind, SYMS_U64 size
 SYMS_API SYMS_TypeNode*
 syms_group_type_mod_from_type(SYMS_Group *group, SYMS_TypeNode *type, SYMS_TypeModifiers mods){
   SYMS_ProfBegin("syms_group_type_mod_from_type");
-  SYMS_ASSERT(syms_thread_lane == 0);
+  SYMS_ASSERT_PARANOID(syms_thread_lane == 0);
   
   SYMS_TypeNode *result = &syms_type_graph_nil;
   
@@ -2003,7 +1972,7 @@ syms_group_type_mod_from_type(SYMS_Group *group, SYMS_TypeNode *type, SYMS_TypeM
 SYMS_API SYMS_TypeNode*
 syms_group_type_ptr_from_type(SYMS_Group *group, SYMS_TypeKind ptr_kind, SYMS_TypeNode *type){
   SYMS_ProfBegin("syms_group_type_ptr_from_type");
-  SYMS_ASSERT(syms_thread_lane == 0);
+  SYMS_ASSERT_PARANOID(syms_thread_lane == 0);
   
   SYMS_TypeNode *result = &syms_type_graph_nil;
   
@@ -2030,7 +1999,7 @@ syms_group_type_ptr_from_type(SYMS_Group *group, SYMS_TypeKind ptr_kind, SYMS_Ty
 SYMS_API SYMS_TypeNode*
 syms_group_type_array_from_type(SYMS_Group *group, SYMS_TypeNode *type, SYMS_U64 count){
   SYMS_ProfBegin("syms_group_type_array_from_type");
-  SYMS_ASSERT(syms_thread_lane == 0);
+  SYMS_ASSERT_PARANOID(syms_thread_lane == 0);
   
   SYMS_TypeNode *result = &syms_type_graph_nil;
   
@@ -2060,7 +2029,7 @@ SYMS_API SYMS_TypeNode*
 syms_group_type_proc_from_type(SYMS_Group *group, SYMS_TypeNode *ret_type, SYMS_TypeNode *this_type,
                                SYMS_TypeNode **param_types, SYMS_U64 param_count){
   SYMS_ProfBegin("syms_group_type_proc_from_type");
-  SYMS_ASSERT(syms_thread_lane == 0);
+  SYMS_ASSERT_PARANOID(syms_thread_lane == 0);
   
   SYMS_TypeNode *result = &syms_type_graph_nil;
   
@@ -2094,7 +2063,7 @@ syms_group_type_proc_from_type(SYMS_Group *group, SYMS_TypeNode *ret_type, SYMS_
 SYMS_API SYMS_TypeNode*
 syms_group_type_member_ptr_from_type(SYMS_Group *group, SYMS_TypeNode *container, SYMS_TypeNode *type){
   SYMS_ProfBegin("syms_group_type_member_ptr_from_type");
-  SYMS_ASSERT(syms_thread_lane == 0);
+  SYMS_ASSERT_PARANOID(syms_thread_lane == 0);
   
   SYMS_TypeNode *result = &syms_type_graph_nil;
   
@@ -2429,7 +2398,7 @@ SYMS_API SYMS_Name2FileIDMap*
 syms_group_file_map(SYMS_Group *group){
   SYMS_ProfBegin("syms_group_file_map");
   if (!group->name_2_file_id_map_is_built){
-    SYMS_ASSERT(syms_thread_lane == 0);
+    SYMS_ASSERT_PARANOID(syms_thread_lane == 0);
     SYMS_ArenaTemp scratch = syms_get_scratch(0, 0);
     
     //- build loose file map
diff --git a/Engine/Source/Runtime/SymsLib/syms/core/group/syms_group.h b/Engine/Source/Runtime/SymsLib/syms/core/group/syms_group.h
--- a/Engine/Source/Runtime/SymsLib/syms/core/group/syms_group.h
+++ b/Engine/Source/Runtime/SymsLib/syms/core/group/syms_group.h
@@ -200,12 +200,14 @@ typedef struct SYMS_Group{
   SYMS_TypeContentBuckets type_content_buckets;
   
   //- one-time fills/builds
-  SYMS_B32 unit_ranges_is_filled;
-  SYMS_B32 type_map_unit_is_filled;
-  SYMS_B32 sec_map_v_is_built;
-  SYMS_B32 sec_map_f_is_built;
-  SYMS_B32 unit_map_is_built;
-  SYMS_B32 name_2_file_id_map_is_built;
+  SYMS_B8 unit_ranges_is_filled;
+  SYMS_B8 type_map_unit_is_filled;
+  SYMS_B8 sec_map_v_is_built;
+  SYMS_B8 sec_map_f_is_built;
+  SYMS_B8 unit_map_is_built;
+  SYMS_B8 name_2_file_id_map_is_built;
+  SYMS_B8 stripped_info_is_filled;
+  SYMS_B8 stripped_info_map_is_built;
   
   SYMS_UnitRangeArray unit_ranges;
   SYMS_UnitAccel *type_map_unit;
@@ -213,6 +215,8 @@ typedef struct SYMS_Group{
   SYMS_SpatialMap1D sec_map_f;
   SYMS_SpatialMap1D unit_map;
   SYMS_Name2FileIDMap name_2_file_id_map;
+  SYMS_StrippedInfoArray stripped_info;
+  SYMS_SpatialMap1D stripped_info_map;
 } SYMS_Group;
 
 ////////////////////////////////
@@ -296,6 +300,8 @@ SYMS_API SYMS_String8       syms_group_file_name_from_id(SYMS_Group *group, SYMS
 
 SYMS_API SYMS_MapAndUnit    syms_group_type_map(SYMS_Group *group);
 
+SYMS_API SYMS_StrippedInfoArray syms_group_stripped_info(SYMS_Group *group);
+
 
 ////////////////////////////////
 // NOTE(allen): Syms Group Address Mapping Functions
@@ -335,6 +341,8 @@ SYMS_API void                syms_group_fetch_line_to_addr_maps_from_uid(SYMS_Gr
 SYMS_API SYMS_LineToAddrMap* syms_group_line_to_addr_map_from_uid_file_id(SYMS_Group *group, SYMS_UnitID uid,
                                                                           SYMS_FileID file_id);
 
+SYMS_API SYMS_SpatialMap1D* syms_group_stripped_info_map(SYMS_Group *group);
+
 //- accelerated versions
 SYMS_API SYMS_U64      syms_group_sec_number_from_voff__accelerated(SYMS_Group *group, SYMS_U64 voff);
 SYMS_API SYMS_U64      syms_group_sec_number_from_foff__accelerated(SYMS_Group *group, SYMS_U64 foff);
diff --git a/Engine/Source/Runtime/SymsLib/syms/core/mach/syms_mach.c b/Engine/Source/Runtime/SymsLib/syms/core/mach/syms_mach.c
--- a/Engine/Source/Runtime/SymsLib/syms/core/mach/syms_mach.c
+++ b/Engine/Source/Runtime/SymsLib/syms/core/mach/syms_mach.c
@@ -10,98 +10,6 @@
 ////////////////////////////////
 //~ NOTE(allen): MACH Format Functions
 
-SYMS_API void
-syms_mach_fat_header_endian_swap_in_place(SYMS_MachFatHeader *x){
-  x->magic     = syms_bswap_u32(x->magic);
-  x->nfat_arch = syms_bswap_u32(x->nfat_arch);
-}
-
-SYMS_API void
-syms_mach_fat_arch_endian_swap_in_place(SYMS_MachFatArch *x){
-  x->cputype    = syms_bswap_u32(x->cputype);
-  x->cpusubtype = syms_bswap_u32(x->cpusubtype);
-  x->offset     = syms_bswap_u32(x->offset);
-  x->size       = syms_bswap_u32(x->size);
-  x->align      = syms_bswap_u32(x->align);
-}
-
-SYMS_API void
-syms_mach_header32_endian_swap_in_place(SYMS_MachHeader32 *x){
-  x->magic      = syms_bswap_u32(x->magic);
-  x->cputype    = syms_bswap_u32(x->cputype);
-  x->cpusubtype = syms_bswap_u32(x->cpusubtype);
-  x->filetype   = syms_bswap_u32(x->filetype);
-  x->ncmds      = syms_bswap_u32(x->ncmds);
-  x->sizeofcmds = syms_bswap_u32(x->sizeofcmds);
-  x->flags      = syms_bswap_u32(x->flags);
-}
-
-SYMS_API void
-syms_mach_header64_endian_swap_in_place(SYMS_MachHeader64 *x){
-  x->magic      = syms_bswap_u32(x->magic);
-  x->cputype    = syms_bswap_u32(x->cputype);
-  x->cpusubtype = syms_bswap_u32(x->cpusubtype);
-  x->filetype   = syms_bswap_u32(x->filetype);
-  x->ncmds      = syms_bswap_u32(x->ncmds);
-  x->sizeofcmds = syms_bswap_u32(x->sizeofcmds);
-  x->flags      = syms_bswap_u32(x->flags);
-  x->reserved   = syms_bswap_u32(x->reserved);
-}
-
-SYMS_API void
-syms_mach_segment_command32_endian_swap_in_place(SYMS_MachSegmentCommand32 *x){
-  x->cmd.type = syms_bswap_u32(x->cmd.type);
-  x->cmd.size = syms_bswap_u32(x->cmd.size);
-  x->vmaddr   = syms_bswap_u32(x->vmaddr);
-  x->vmsize   = syms_bswap_u32(x->vmsize);
-  x->fileoff  = syms_bswap_u32(x->fileoff);
-  x->filesize = syms_bswap_u32(x->filesize);
-  x->maxprot  = syms_bswap_u32(x->maxprot);
-  x->initprot = syms_bswap_u32(x->initprot);
-  x->nsects   = syms_bswap_u32(x->nsects);
-  x->flags    = syms_bswap_u32(x->flags);
-}
-
-SYMS_API void
-syms_mach_segment_command64_endian_swap_in_place(SYMS_MachSegmentCommand64 *x){
-  x->cmd.type = syms_bswap_u32(x->cmd.type);
-  x->cmd.size = syms_bswap_u32(x->cmd.size);
-  x->vmaddr   = syms_bswap_u64(x->vmaddr);
-  x->vmsize   = syms_bswap_u64(x->vmsize);
-  x->fileoff  = syms_bswap_u64(x->fileoff);
-  x->filesize = syms_bswap_u64(x->filesize);
-  x->maxprot  = syms_bswap_u32(x->maxprot);
-  x->initprot = syms_bswap_u32(x->initprot);
-  x->nsects   = syms_bswap_u32(x->nsects);
-  x->flags    = syms_bswap_u32(x->flags);
-}
-
-SYMS_API void
-syms_mach_section32_endian_swap_in_place(SYMS_MachSection32 *x){
-  x->addr         = syms_bswap_u32(x->addr);
-  x->size         = syms_bswap_u32(x->size);
-  x->offset       = syms_bswap_u32(x->offset);
-  x->align        = syms_bswap_u32(x->align);
-  x->relocoff     = syms_bswap_u32(x->relocoff);
-  x->nreloc       = syms_bswap_u32(x->nreloc);
-  x->flags        = syms_bswap_u32(x->flags);
-  x->reserved1    = syms_bswap_u32(x->reserved1);
-  x->reserved2    = syms_bswap_u32(x->reserved2);
-}
-
-SYMS_API void
-syms_mach_section64_endian_swap_in_place(SYMS_MachSection64 *x){
-  x->addr         = syms_bswap_u64(x->addr);
-  x->size         = syms_bswap_u64(x->size);
-  x->offset       = syms_bswap_u32(x->offset);
-  x->align        = syms_bswap_u32(x->align);
-  x->relocoff     = syms_bswap_u32(x->relocoff);
-  x->nreloc       = syms_bswap_u32(x->nreloc);
-  x->flags        = syms_bswap_u32(x->flags);
-  x->reserved1    = syms_bswap_u32(x->reserved1);
-  x->reserved2    = syms_bswap_u32(x->reserved2);
-}
-
 SYMS_API void
 syms_mach_header64_from_header32(SYMS_MachHeader64 *dst, SYMS_MachHeader32 *header32){
   dst->magic = header32->magic;
@@ -123,16 +31,35 @@ syms_mach_nlist64_from_nlist32(SYMS_MachNList64 *dst, SYMS_MachNList32 *nlist32)
   dst->n_value = nlist32->n_value;
 }
 
-SYMS_API SYMS_Arch
-syms_mach_arch_from_cputype(SYMS_MachCpuType cputype){
-  SYMS_Arch result = SYMS_Arch_Null;
-  switch (cputype){
-    case SYMS_MachCpuType_X86:    result = SYMS_Arch_X86;   break;
-    case SYMS_MachCpuType_X86_64: result = SYMS_Arch_X64;   break;
-    case SYMS_MachCpuType_ARM:    result = SYMS_Arch_ARM32; break;
-    case SYMS_MachCpuType_ARM64:  result = SYMS_Arch_ARM;   break;
-  }
-  return(result);
+SYMS_API void
+syms_mach_segment64_from_segment32(SYMS_MachSegmentCommand64 *dst, SYMS_MachSegmentCommand32 *seg32)
+{
+  dst->cmd = seg32->cmd;
+  syms_memmove(&dst->segname[0], &seg32->segname[0], sizeof(seg32->segname));
+  dst->vmaddr = seg32->vmaddr;
+  dst->vmsize = seg32->vmsize;
+  dst->fileoff = seg32->fileoff;
+  dst->filesize = seg32->filesize;
+  dst->maxprot = seg32->maxprot;
+  dst->initprot = seg32->initprot;
+  dst->nsects = seg32->nsects;
+  dst->flags = seg32->flags;
+}
+
+SYMS_API void
+syms_mach_section64_from_section32(SYMS_MachSection64 *dst, SYMS_MachSection32 *sect32)
+{
+  syms_memmove(&dst->sectname[0], &sect32->sectname[0], sizeof(sect32->sectname));
+  syms_memmove(&dst->segname[0], &sect32->segname[0], sizeof(sect32->segname));
+  dst->addr = sect32->addr;
+  dst->size = sect32->size;
+  dst->offset = sect32->offset;
+  dst->align = sect32->align;
+  dst->relocoff = sect32->relocoff;
+  dst->nreloc = sect32->nreloc;
+  dst->flags = sect32->flags;
+  dst->reserved1 = sect32->reserved1;
+  dst->reserved2 = sect32->reserved2;
 }
 
 #endif // SYMS_MACH_C
diff --git a/Engine/Source/Runtime/SymsLib/syms/core/mach/syms_mach.h b/Engine/Source/Runtime/SymsLib/syms/core/mach/syms_mach.h
--- a/Engine/Source/Runtime/SymsLib/syms/core/mach/syms_mach.h
+++ b/Engine/Source/Runtime/SymsLib/syms/core/mach/syms_mach.h
@@ -4,17 +4,12 @@
 
 /*
 ** Reference materials:
-  ** machine.h https://opensource.apple.com/source/xnu/xnu-4570.41.2/osfmk/mach/machine.h.auto.html
-  ** loader.h  https://opensource.apple.com/source/xnu/xnu-4570.71.2/EXTERNAL_HEADERS/mach-o/loader.h.auto.html
-  ** nlist.h   https://opensource.apple.com/source/xnu/xnu-201/EXTERNAL_HEADERS/mach-o/nlist.h.auto.html
-  ** stab.h    https://opensource.apple.com/source/xnu/xnu-4903.241.1/EXTERNAL_HEADERS/mach-o/stab.h.auto.html
+** machine.h https://opensource.apple.com/source/xnu/xnu-4570.41.2/osfmk/mach/machine.h.auto.html
+** loader.h  https://opensource.apple.com/source/xnu/xnu-4570.71.2/EXTERNAL_HEADERS/mach-o/loader.h.auto.html
+** nlist.h   https://opensource.apple.com/source/xnu/xnu-201/EXTERNAL_HEADERS/mach-o/nlist.h.auto.html
+** stab.h    https://opensource.apple.com/source/xnu/xnu-4903.241.1/EXTERNAL_HEADERS/mach-o/stab.h.auto.html
 */
 
-////////////////////////////////
-//~ NOTE(allen): Generated
-
-#include "syms/core/generated/syms_meta_mach.h"
-
 ////////////////////////////////
 //~ NOTE(allen): Constants from machine.h
 
@@ -27,6 +22,8 @@
 ////////////////////////////////
 //~ NOTE(allen): Types & Constants from loader.h
 
+typedef SYMS_S32 SYMS_MachVMProt;
+
 //- magic numbers
 #define SYMS_MACH_MAGIC_32 0xFEEDFACE
 #define SYMS_MACH_CIGAM_32 0xCEFAEDFE
@@ -44,176 +41,19 @@
 #define SYMS_MACH_FAT_IS_SWAPPED(x)    ((x) == SYMS_MACH_FAT_CIGAM)
 
 //- types
-typedef struct SYMS_MachLCStr{
-  SYMS_U32 offset;
-} SYMS_MachLCStr;
-
-typedef struct SYMS_MachFatHeader{
-  SYMS_U32 magic;
-  SYMS_U32 nfat_arch;
-} SYMS_MachFatHeader;
-
-typedef struct SYMS_MachFatArch{
-  SYMS_MachCpuType cputype;
-  SYMS_MachCpuSubtype cpusubtype;
-  SYMS_U32 offset;
-  SYMS_U32 size;
-  SYMS_U32 align;
-} SYMS_MachFatArch;
-
-typedef struct SYMS_MachHeader32{
-  SYMS_U32 magic;
-  SYMS_MachCpuType cputype;
-  SYMS_MachCpuSubtype cpusubtype;
-  SYMS_MachFiletype filetype;
-  SYMS_U32 ncmds;
-  SYMS_U32 sizeofcmds;
-  SYMS_MachFlags flags;
-} SYMS_MachHeader32;
-typedef struct SYMS_MachHeader64{
-  SYMS_U32 magic;
-  SYMS_MachCpuType cputype;
-  SYMS_MachCpuSubtype cpusubtype;
-  SYMS_MachFiletype filetype;
-  SYMS_U32 ncmds;
-  SYMS_U32 sizeofcmds;
-  SYMS_MachFlags flags;
-  SYMS_U32 reserved;
-} SYMS_MachHeader64;
-
-typedef struct SYMS_MachLoadCommand{
-  SYMS_MachLoadCommandType type;
-  SYMS_U32 size;
-} SYMS_MachLoadCommand;
-
-typedef SYMS_S32 SYMS_MachVMProt;
-
-typedef struct SYMS_MachSegmentCommand32{
-  SYMS_MachLoadCommand cmd;
-  SYMS_U8 segname[16];
-  SYMS_U32 vmaddr;
-  SYMS_U32 vmsize;
-  SYMS_U32 fileoff;
-  SYMS_U32 filesize;
-  SYMS_MachVMProt maxprot;
-  SYMS_MachVMProt initprot;
-  SYMS_U32 nsects;
-  SYMS_U32 flags;
-} SYMS_MachSegmentCommand32;
-
-typedef struct SYMS_MachSegmentCommand64{
-  SYMS_MachLoadCommand cmd;
-  SYMS_U8 segname[16];
-  SYMS_U64 vmaddr;
-  SYMS_U64 vmsize;
-  SYMS_U64 fileoff;
-  SYMS_U64 filesize;
-  SYMS_MachVMProt maxprot;
-  SYMS_MachVMProt initprot;
-  SYMS_U32 nsects;
-  SYMS_U32 flags;
-} SYMS_MachSegmentCommand64;
-
-typedef struct SYMS_MachSection32{
-  SYMS_U8 sectname[16];
-  SYMS_U8 segname[16];
-  SYMS_U32 addr;
-  SYMS_U32 size;
-  SYMS_U32 offset;
-  SYMS_U32 align;
-  SYMS_U32 relocoff;
-  SYMS_U32 nreloc;
-  SYMS_U32 flags;
-  SYMS_U32 reserved1;
-  SYMS_U32 reserved2;
-} SYMS_MachSection32;
-
-typedef struct SYMS_MachSection64{
-  SYMS_U8 sectname[16];
-  SYMS_U8 segname[16];
-  SYMS_U64 addr;
-  SYMS_U64 size;
-  SYMS_U32 offset;
-  SYMS_U32 align;
-  SYMS_U32 relocoff;
-  SYMS_U32 nreloc;
-  SYMS_U32 flags;
-  SYMS_U32 reserved1;
-  SYMS_U32 reserved2;
-} SYMS_MachSection64;
-
-// 128-bit unique number generated by linker to identify objects
-typedef struct SYMS_MachUUID{
-  SYMS_U32 cmd;
-  SYMS_U32 cmdsize;
-  SYMS_U8 uuid[16];
-} SYMS_MachUUID;
-
-typedef struct SYMS_MachDylib{
-  SYMS_MachLCStr name;
-  SYMS_U32 timestamp;
-  SYMS_U32 current_version;
-  SYMS_U32 compatability_version;
-} SYMS_MachDylib;
-
-typedef struct SYMS_MachDylibCommand{
-  SYMS_U32 cmd;
-  SYMS_U32 cmdsize;
-  SYMS_MachDylib dylib;
-} SYMS_MachDylibCommand;
-
-typedef struct SYMS_MachDyldInfoCommand{
-  SYMS_U32 cmd;
-  SYMS_U32 cmdsize;
-  SYMS_U32 rebase_off;
-  SYMS_U32 rebase_size;
-  SYMS_U32 bind_off;
-  SYMS_U32 bind_size;
-  SYMS_U32 weak_bind_off;
-  SYMS_U32 weak_bind_size;
-  SYMS_U32 lazy_bind_off;
-  SYMS_U32 lazy_bind_size;
-  SYMS_U32 export_off;
-  SYMS_U32 export_size;
-} SYMS_MachDyldInfoCommand;
-
-// command for symbol table, format for entries is defined in list.h and stab.h
-typedef struct SYMS_MachSymtabCommand{
-  SYMS_U32 cmd;
-  SYMS_U32 cmdsize;
-  SYMS_U32 symoff;  // offset from the image base to array of SYMS_MachNList32
-  SYMS_U32 nsyms;   // symbol count
-  SYMS_U32 stroff;  // offset from the image base to string table
-  SYMS_U32 strsize; // string table size
-} SYMS_MachSymtabCommand;
-
-typedef struct SYMS_MachDySymtabCommand{
-  SYMS_U32 cmd;
-  SYMS_U32 cmdsize;
-  SYMS_U32 ilocalsym;
-  SYMS_U32 nlocalsym;
-  SYMS_U32 iextdefsym;
-  SYMS_U32 nextdefsym;
-  SYMS_U32 iundefsym;
-  SYMS_U32 nundefsym;
-  SYMS_U32 tocoff;
-  SYMS_U32 ntoc;
-  SYMS_U32 modtaboff;
-  SYMS_U32 nmodtab;
-  SYMS_U32 extrefsymoff;
-  SYMS_U32 nextrefsyms;
-  SYMS_U32 indirectsymoff;
-  SYMS_U32 nindirectsyms;
-  SYMS_U32 extreloff;
-  SYMS_U32 nextrel;
-  SYMS_U32 locreloff;
-  SYMS_U32 nlocrel;
-} SYMS_MachDySymtabCommand;
 
 #define SYMS_MACH_VERSION_NIBBLE_A(x) (((x) & 0xFFFF0000) >> 16)
 #define SYMS_MACH_VERSION_NIBBLE_B(x) (((x) & 0x0000FF00) >> 8)
 #define SYMS_MACH_VERSION_NIBBLE_C(x) (((x) & 0x000000FF) >> 0)
 
+////////////////////////////////
+//~ NOTE(allen): Generated
+
+#include "syms/core/generated/syms_meta_mach.h"
+
+////////////////////////////////
+//~ NOTE(rjf): More Types & Constants from loader.h
+
 typedef struct SYMS_MachBuildVersionCommand{
   SYMS_U32 cmd;
   SYMS_U32 cmdsize;
@@ -236,40 +76,15 @@ typedef struct SYMS_MachBuildToolVersion{
 #define SYMS_MACH_NLIST_TYPE(x) ((x) & 0x0E)
 #define SYMS_MACH_NLIST_EXT(x)  ((x) & 0x01)
 
-typedef struct SYMS_MachNList32{
-  SYMS_U32 n_strx;
-  SYMS_U8 n_type;
-  SYMS_U8 n_sect;
-  SYMS_U16 n_desc;
-  SYMS_U32 n_value;
-} SYMS_MachNList32;
-
-typedef struct SYMS_MachNList64{
-  SYMS_U32 n_strx;
-  SYMS_U8 n_type;
-  SYMS_U8 n_sect;
-  SYMS_U16 n_desc;
-  SYMS_U64 n_value;
-} SYMS_MachNList64;
-
 ////////////////////////////////
 //~ NOTE(allen): MACH Format Functions
 
 SYMS_C_LINKAGE_BEGIN
 
-SYMS_API void      syms_mach_fat_header_endian_swap_in_place(SYMS_MachFatHeader *x);
-SYMS_API void      syms_mach_fat_arch_endian_swap_in_place(SYMS_MachFatArch *x);
-SYMS_API void      syms_mach_header32_endian_swap_in_place(SYMS_MachHeader32 *x);
-SYMS_API void      syms_mach_header64_endian_swap_in_place(SYMS_MachHeader64 *x);
-SYMS_API void      syms_mach_segment_command32_endian_swap_in_place(SYMS_MachSegmentCommand32 *x);
-SYMS_API void      syms_mach_segment_command64_endian_swap_in_place(SYMS_MachSegmentCommand64 *x);
-SYMS_API void      syms_mach_section32_endian_swap_in_place(SYMS_MachSection32 *section32);
-SYMS_API void      syms_mach_section64_endian_swap_in_place(SYMS_MachSection64 *section64);
-
-SYMS_API void      syms_mach_header64_from_header32(SYMS_MachHeader64 *dst, SYMS_MachHeader32 *header32);
-SYMS_API void      syms_mach_nlist64_from_nlist32(SYMS_MachNList64 *dst, SYMS_MachNList32 *nlist32);
-
-SYMS_API SYMS_Arch syms_mach_arch_from_cputype(SYMS_MachCpuType cputype);
+SYMS_API void syms_mach_header64_from_header32(SYMS_MachHeader64 *dst, SYMS_MachHeader32 *header32);
+SYMS_API void syms_mach_nlist64_from_nlist32(SYMS_MachNList64 *dst, SYMS_MachNList32 *nlist32);
+SYMS_API void syms_mach_segment64_from_segment32(SYMS_MachSegmentCommand64 *dst, SYMS_MachSegmentCommand32 *seg32);
+SYMS_API void syms_mach_section64_from_section32(SYMS_MachSection64 *dst, SYMS_MachSection32 *sect32);
 
 SYMS_C_LINKAGE_END
 
diff --git a/Engine/Source/Runtime/SymsLib/syms/core/mach/syms_mach_parser.c b/Engine/Source/Runtime/SymsLib/syms/core/mach/syms_mach_parser.c
--- a/Engine/Source/Runtime/SymsLib/syms/core/mach/syms_mach_parser.c
+++ b/Engine/Source/Runtime/SymsLib/syms/core/mach/syms_mach_parser.c
@@ -70,6 +70,13 @@ syms_mach_bin_from_base_range(SYMS_Arena *arena, void *base, SYMS_U64Range range
     SYMS_MachSectionNode *section_last = 0;
     SYMS_U32 section_count = 0;
     
+    SYMS_U64Range regular_bind_range = syms_make_u64_range(0,0);
+    SYMS_U64Range lazy_bind_range    = syms_make_u64_range(0,0);
+    SYMS_U64Range weak_bind_range    = syms_make_u64_range(0,0);
+    SYMS_U64Range export_range       = syms_make_u64_range(0,0);
+    
+    SYMS_MachDylibList dylib_list; syms_memzero_struct(&dylib_list);
+    
     SYMS_U64 next_cmd_off = after_header_off;
     for (SYMS_U32 i = 0; i < header.ncmds; i += 1){
       //- align read offset
@@ -123,6 +130,35 @@ syms_mach_bin_from_base_range(SYMS_Arena *arena, void *base, SYMS_U64Range range
             section_count += 1;
           }
         }break;
+        
+        case SYMS_MachLoadCommandType_DYLD_INFO_ONLY: 
+        {
+          SYMS_MachDyldInfoCommand dyld;
+          syms_based_range_read_struct(base, range, cmd_off, &dyld);
+          if (is_swapped){
+            syms_mach_dyld_info_command_endian_swap_in_place(&dyld);
+          }
+          
+          regular_bind_range = syms_make_u64_inrange(range, dyld.bind_off, dyld.bind_size);
+          lazy_bind_range    = syms_make_u64_inrange(range, dyld.lazy_bind_off, dyld.lazy_bind_size);
+          weak_bind_range    = syms_make_u64_inrange(range, dyld.weak_bind_off, dyld.weak_bind_size);
+          export_range       = syms_make_u64_inrange(range, dyld.export_off, dyld.export_size);
+        }break;
+        
+        case SYMS_MachLoadCommandType_ID_DYLIB:
+        case SYMS_MachLoadCommandType_LOAD_DYLIB: 
+        {
+          SYMS_MachDylibCommand cmd; syms_memzero_struct(&cmd);
+          syms_based_range_read_struct(base, range, cmd_off, &cmd);
+          SYMS_MachDylib *dylib = &cmd.dylib;
+          if (is_swapped){
+            syms_mach_dylib_endian_swap_in_place(dylib);
+          }
+          
+          SYMS_U64 name_offset = cmd_off + dylib->name.offset;
+          SYMS_U64Range name = syms_make_u64_inrange(range, name_offset, (cmd_off + lc.size) - name_offset);
+          syms_mach_dylib_list_push(scratch.arena, &dylib_list, dylib, name);
+        }break;
       }
       
       next_cmd_off = cmd_off + lc.size;
@@ -151,13 +187,24 @@ syms_mach_bin_from_base_range(SYMS_Arena *arena, void *base, SYMS_U64Range range
     }
     
     //- fill result
-    result = syms_push_array(arena, SYMS_MachBinAccel, 1);
-    result->format = SYMS_FileFormat_MACH;
-    result->arch = syms_mach_arch_from_cputype(header.cputype);
-    result->segment_count = segment_count;
-    result->segments = segments;
-    result->section_count = section_count;
-    result->sections = sections;
+    result                     = syms_push_array(arena, SYMS_MachBinAccel, 1);
+    result->format             = SYMS_FileFormat_MACH;
+    result->arch               = syms_mach_arch_from_cputype(header.cputype);
+    result->segment_count      = segment_count;
+    result->segments           = segments;
+    result->section_count      = section_count;
+    result->sections           = sections;
+    result->regular_bind_range = regular_bind_range;
+    result->lazy_bind_range    = lazy_bind_range;
+    result->weak_bind_range    = weak_bind_range;
+    result->export_range       = export_range;
+    { // convert dylib list to array
+      result->dylib_count = 0;
+      result->dylibs = syms_push_array(arena, SYMS_MachParsedDylib, dylib_list.count);
+      for (SYMS_MachDylibNode *n = dylib_list.first; n != 0; n = n->next){
+        result->dylibs[result->dylib_count++] = n->data;
+      }
+    }
     
     syms_release_scratch(scratch);
   }
@@ -346,10 +393,8 @@ syms_mach_sec_info_array_from_bin(SYMS_Arena *arena, SYMS_String8 data, SYMS_Mac
   }
   return array;
 }
-
 SYMS_API SYMS_U64
-syms_mach_default_vbase_from_bin(SYMS_MachBinAccel *bin)
-{
+syms_mach_default_vbase_from_bin(SYMS_MachBinAccel *bin){
   // TODO(rjf): @nick verify
   SYMS_U64 min_vbase = 0;
   for(SYMS_U64 segment_idx = 0; segment_idx < bin->segment_count; segment_idx += 1)
@@ -363,4 +408,299 @@ syms_mach_default_vbase_from_bin(SYMS_MachBinAccel *bin)
   return min_vbase;
 }
 
+////////////////////////////////
+
+SYMS_API void
+syms_mach_dylib_list_push(SYMS_Arena *arena, SYMS_MachDylibList *list, SYMS_MachDylib *dylib, SYMS_U64Range name){
+  SYMS_MachDylibNode *node = syms_push_array(arena, SYMS_MachDylibNode, 1);
+  node->data.header = *dylib;
+  node->data.name = name;
+  node->next = 0;
+  SYMS_QueuePush(list->first, list->last, node);
+  list->count += 1;
+}
+
+////////////////////////////////
+
+SYMS_API SYMS_MachBindList
+syms_mach_binds_from_base_range(SYMS_Arena *arena, void *base, SYMS_U64Range range, SYMS_U32 address_size, SYMS_MachBindTable bind_type){
+  SYMS_U64 read_offset = 0;
+  SYMS_MachBindList list; syms_memzero_struct(&list);
+  SYMS_MachBind state;    syms_memzero_struct(&state);
+  for (;;){
+    SYMS_U8 encoded_byte = 0;
+    read_offset += syms_based_range_read_struct(base, range, read_offset, &encoded_byte);
+    
+    SYMS_U8 opcode = encoded_byte & SYMS_MachBindOpcode_MASK;
+    SYMS_U8 imm    = encoded_byte & SYMS_MachBindOpcode_IMM_MASK;
+    
+    switch (opcode){
+      case SYMS_MachBindOpcode_DONE: goto exit;
+      
+      case SYMS_MachBindOpcode_SET_DYLIB_ORDINAL_IMM:
+      {
+        state.dylib = imm;
+      } break;
+      case SYMS_MachBindOpcode_SET_DYLIB_ORDINAL_ULEB:
+      {
+        read_offset += syms_based_range_read_uleb128(base, range, read_offset, &state.dylib);
+      } break;
+      case SYMS_MachBindOpcode_SET_DYLIB_SPECIAL_IMM:
+      {
+        SYMS_ASSERT_PARANOID("TODO: SET_DYLIB_SPECIAL_IMM");
+      } break;
+      case SYMS_MachBindOpcode_SET_SYMBOL_TRAILING_FLAGS_IMM:
+      {
+        state.flags = imm;
+        state.symbol_name = syms_based_range_read_string(base, range, read_offset);
+        read_offset += (state.symbol_name.size + 1);
+      } break;
+      case SYMS_MachBindOpcode_SET_TYPE_IMM:
+      {
+        state.type = imm;
+      } break;
+      case SYMS_MachBindOpcode_SET_ADDEND_SLEB:
+      {
+        read_offset += syms_based_range_read_sleb128(base, range, read_offset, &state.addend);
+      } break;
+      case SYMS_MachBindOpcode_SET_SEGMENT_AND_OFFSET_ULEB:
+      {
+        state.segment = imm;
+        read_offset += syms_based_range_read_uleb128(base, range, read_offset, &state.segment_offset);
+      } break;
+      case SYMS_MachBindOpcode_ADD_ADDR_ULEB:
+      {
+        SYMS_U64 addend = 0;
+        read_offset += syms_based_range_read_uleb128(base, range, read_offset, &addend);
+        state.segment_offset += addend;
+      } break;
+      case SYMS_MachBindOpcode_DO_BIND_ADD_ADDR_ULEB:
+      {
+        SYMS_ASSERT_PARANOID(!"TODO: DO_BIND_ADD_ADDR_ULEB");
+      } goto do_bind;
+      case SYMS_MachBindOpcode_DO_BIND_ADD_ADDR_IMM_SCALED:
+      {
+        SYMS_ASSERT_PARANOID(!"TODO: DO_BIND_ADD_ADDR_IMM_SCALED");
+      } goto do_bind;
+      case SYMS_MachBindOpcode_DO_BIND_ULEB_TIMES_SKIPPING_ULEB:
+      {
+        SYMS_ASSERT_PARANOID(!"TODO: DO_BIND_ULEB_TIMES_SKIPPING_ULEB");
+      } goto do_bind;
+      case SYMS_MachBindOpcode_DO_BIND:
+      {
+      } goto do_bind;
+      do_bind:
+      {
+        if (bind_type == SYMS_MachBindTable_LAZY){
+          for (;read_offset < syms_u64_range_size(range);){
+            SYMS_U8 value = 0;
+            SYMS_U64 read_size = syms_based_range_read_struct(base, range, read_offset, &value);
+            if (value != 0){
+              break;
+            }
+            read_offset += read_size;
+          }
+        }
+        SYMS_MachBindNode *node = syms_push_array(arena, SYMS_MachBindNode, 1);
+        node->data = state;
+        node->next = 0;
+        SYMS_QueuePush(list.first, list.last, node);
+        list.count += 1;
+      } break;
+    }
+  }
+  exit:;
+  
+  return list;
+}
+
+SYMS_API SYMS_ImportArray
+syms_mach_imports_from_bin(SYMS_Arena *arena, SYMS_String8 data, SYMS_MachBinAccel *bin){
+  SYMS_ArenaTemp scratch = syms_get_scratch(&arena, 1);
+
+  // read and copy library names
+  SYMS_String8 *dylib_names = syms_push_array(scratch.arena, SYMS_String8, bin->dylib_count);
+  for (SYMS_U32 i = 0; i < bin->dylib_count; ++i){
+    SYMS_String8 name = syms_based_range_read_string((void*)data.str, bin->dylibs[i].name, 0);
+    dylib_names[i] = syms_push_string_copy(arena, name);
+  }
+  
+  SYMS_U64 address_size = syms_address_size_from_arch(bin->arch);
+  SYMS_MachBindList binds[3];
+  binds[SYMS_MachBindTable_REGULAR] = syms_mach_binds_from_base_range(scratch.arena, (void*)data.str, bin->regular_bind_range, address_size, SYMS_MachBindTable_REGULAR);
+  binds[SYMS_MachBindTable_LAZY]    = syms_mach_binds_from_base_range(scratch.arena, (void*)data.str, bin->lazy_bind_range, address_size, SYMS_MachBindTable_LAZY);
+  binds[SYMS_MachBindTable_WEAK]    = syms_mach_binds_from_base_range(scratch.arena, (void*)data.str, bin->weak_bind_range, address_size, SYMS_MachBindTable_WEAK);
+  
+  SYMS_U64 total_bind_count = 0;
+  for (SYMS_U32 i = 0; i < SYMS_ARRAY_SIZE(binds); ++i){
+    total_bind_count += binds[i].count;
+  }
+
+  SYMS_ImportArray import_array; 
+  import_array.count = 0;
+  import_array.imports = syms_push_array(arena, SYMS_Import, total_bind_count);
+  for (SYMS_U32 i = 0; i < SYMS_ARRAY_SIZE(binds); ++i){
+    for (SYMS_MachBindNode *node = binds[i].first; node != 0; node = node->next){
+      SYMS_MachBind *bind = &node->data;
+      SYMS_Import *import = &import_array.imports[import_array.count++];
+      import->name = syms_push_string_copy(arena, bind->symbol_name);
+      if (bind->dylib < bin->dylib_count){
+        import->library_name = dylib_names[bind->dylib];
+      } else{
+        import->library_name = syms_str8(0,0);
+      }
+      import->ordinal = 0;
+    }
+  }
+  
+  syms_release_scratch(scratch);
+  return import_array;
+}
+
+static SYMS_MachExport *
+syms_mach_parse_export_node(SYMS_Arena *arena, void *base, SYMS_U64Range range, SYMS_String8 name, SYMS_U64 read_offset){
+  SYMS_MachExport *node = 0;
+  
+  SYMS_U64 export_size = 0;
+  SYMS_U64 read_size = syms_based_range_read_uleb128(base, range, read_offset, &export_size);
+  read_offset += read_size;
+  
+  if (read_size > 0){
+    node = syms_push_array_zero(arena, SYMS_MachExport, 1);
+    node->name = name;
+    
+    SYMS_B32 is_export_info = export_size != 0;
+    if (is_export_info){
+      SYMS_U64Range export_range = syms_make_u64_inrange(range, read_offset, export_size);
+      SYMS_U64 export_read_offset = 0;
+      export_read_offset += syms_based_range_read_uleb128(base, export_range, export_read_offset, &node->flags);
+      if (node->flags & SYMS_MachExportSymbolFlags_REEXPORT){
+        export_read_offset += syms_based_range_read_uleb128(base, export_range, export_read_offset, &node->dylib_ordinal);
+        node->import_name = syms_based_range_read_string(base, export_range, export_read_offset);
+        export_read_offset += (node->import_name.size + 1);
+      } else{
+        export_read_offset += syms_based_range_read_uleb128(base, export_range, export_read_offset, &node->address);
+        if (node->flags & SYMS_MachExportSymbolFlags_STUB_AND_RESOLVED){
+          export_read_offset += syms_based_range_read_uleb128(base, export_range, export_read_offset, &node->resolver);
+        }
+      }
+    }
+    
+    read_offset += syms_based_range_read_struct(base, range, read_offset, &node->child_count);
+    if (node->child_count > 0){
+      node->children = syms_push_array(arena, SYMS_MachExport *, node->child_count);
+      for (SYMS_U8 child_idx = 0; child_idx < node->child_count; child_idx += 1){
+        SYMS_String8 name = syms_based_range_read_string(base, range, read_offset);
+        read_offset += (name.size + 1);
+        SYMS_U64 child_offset = 0;
+        read_offset += syms_based_range_read_uleb128(base, range, read_offset, &child_offset);
+        node->children[child_idx] = syms_mach_parse_export_node(arena, base, range, name, child_offset);
+      }
+    }
+  }
+  
+  return node;
+}
+
+SYMS_API SYMS_MachExport *
+syms_build_mach_export_trie(SYMS_Arena *arena, void *base, SYMS_U64Range range){
+  // TODO(nick): rewrite function to use arena for recursion.
+  return syms_mach_parse_export_node(arena, base, range, syms_str8_lit(""), 0);
+}
+
+SYMS_API SYMS_ExportArray
+syms_mach_exports_from_bin(SYMS_Arena *arena, SYMS_String8 data, SYMS_MachBinAccel *bin){
+  SYMS_ArenaTemp scratch = syms_get_scratch(&arena, 1);
+
+  // read and copy library names
+  SYMS_String8 *dylib_names = syms_push_array(scratch.arena, SYMS_String8, bin->dylib_count);
+  for (SYMS_U32 i = 0; i < bin->dylib_count; ++i){
+    SYMS_String8 name = syms_based_range_read_string((void*)data.str, bin->dylibs[i].name, 0);
+    dylib_names[i] = syms_push_string_copy(arena, name);
+  }
+  
+  SYMS_MachExport *root_export = syms_build_mach_export_trie(scratch.arena, (void*)data.str, bin->export_range);
+  SYMS_ExportNode *first_export = 0;
+  SYMS_ExportNode *last_export = 0;
+  SYMS_U32 export_count = 0;
+  if (root_export){
+    SYMS_MachExportFrame *stack = 0;
+    SYMS_MachExportFrame *free_frames = 0;
+    
+    SYMS_MachExportFrame *frame = syms_push_array_zero(arena, SYMS_MachExportFrame, 1);
+    frame->node = root_export;
+    SYMS_StackPush(stack, frame);
+    
+    for (;;){
+      if (frame->child_idx >= frame->node->child_count){
+        if (frame->node->child_count == 0){ // leaf
+          SYMS_ArenaTemp temp = syms_arena_temp_begin(scratch.arena);
+          SYMS_String8List name_parts; syms_memzero_struct(&name_parts);
+          for (SYMS_MachExportFrame *f = frame; f != 0; f = f->next){
+            syms_string_list_push_front(temp.arena, &name_parts, f->node->name);
+          }
+          SYMS_StringJoin join; syms_memzero_struct(&join);
+          SYMS_String8 name = syms_string_list_join(arena, &name_parts, &join);
+          syms_arena_temp_end(temp);
+          
+          SYMS_ExportNode *export_node = syms_push_array(scratch.arena, SYMS_ExportNode, 1); 
+          SYMS_QueuePush(first_export, last_export, export_node);
+          export_count += 1;
+          SYMS_Export *exp = &export_node->data;
+          if (frame->node->flags & SYMS_MachExportSymbolFlags_REEXPORT){
+            // TODO(nick): need dylib with reexports to test this code path
+            if (frame->node->dylib_ordinal > 0 && frame->node->dylib_ordinal < bin->dylib_count){
+              exp->name = name;
+              exp->address = 0;
+              exp->ordinal = 0;
+              exp->forwarder_library_name = dylib_names[frame->node->dylib_ordinal-1];
+              exp->forwarder_import_name = syms_push_string_copy(arena, frame->node->import_name);
+            } else{
+              SYMS_ASSERT_PARANOID(!"invalid dylib ordinal");
+            }
+          } else{
+            exp->name = name;
+            exp->address = frame->node->address;
+            exp->ordinal = 0;
+            exp->forwarder_library_name = syms_str8(0,0);
+            exp->forwarder_import_name = syms_str8(0,0);
+          }
+        }
+        // pop current frame 
+        SYMS_MachExportFrame *to_free = frame;
+        frame = SYMS_StackPop(stack);
+        // push frame to free list for later reuse
+        to_free->next = 0;
+        SYMS_StackPush(free_frames, to_free);
+        if (frame == 0){
+          break; // we traversed entire trie, leave
+        }
+      }
+      if (frame->child_idx < frame->node->child_count){
+        // push frame with child node
+        SYMS_MachExportFrame *new_frame = free_frames;
+        if (new_frame){
+          SYMS_StackPop(free_frames);
+        } else{
+          new_frame = syms_push_array_zero(arena, SYMS_MachExportFrame, 1);
+        }
+        new_frame->next = 0;
+        new_frame->child_idx = 0;
+        new_frame->node = frame->node->children[frame->child_idx++];
+        frame = SYMS_StackPush(stack, new_frame);
+      }
+    }
+  }
+  
+  SYMS_ExportArray export_array;
+  export_array.count = 0;
+  export_array.exports = syms_push_array(arena, SYMS_Export, export_count);
+  for (SYMS_ExportNode *n = first_export; n != 0; n = n->next){
+    export_array.exports[export_array.count++] = n->data;
+  }
+  
+  syms_release_scratch(scratch);
+  return export_array;
+}
+
 #endif // SYMS_MACH_PARSER_C
diff --git a/Engine/Source/Runtime/SymsLib/syms/core/mach/syms_mach_parser.h b/Engine/Source/Runtime/SymsLib/syms/core/mach/syms_mach_parser.h
--- a/Engine/Source/Runtime/SymsLib/syms/core/mach/syms_mach_parser.h
+++ b/Engine/Source/Runtime/SymsLib/syms/core/mach/syms_mach_parser.h
@@ -37,8 +37,107 @@ typedef struct SYMS_MachBinAccel{
   
   SYMS_MachSection64 *sections;
   SYMS_U32 section_count;
+
+  SYMS_U64Range regular_bind_range;
+  SYMS_U64Range lazy_bind_range;
+  SYMS_U64Range weak_bind_range;
+  SYMS_U64Range export_range;
+
+  SYMS_U32 dylib_count;
+  struct SYMS_MachParsedDylib *dylibs;
 } SYMS_MachBinAccel;
 
+////////////////////////////////
+// Dylib
+
+typedef struct SYMS_MachParsedDylib
+{
+  SYMS_MachDylib header;
+  SYMS_U64Range name;
+} SYMS_MachParsedDylib;
+
+
+typedef struct SYMS_MachDylibNode
+{
+  struct SYMS_MachDylibNode *next;
+  SYMS_MachParsedDylib data;
+} SYMS_MachDylibNode;
+
+typedef struct
+{
+  SYMS_MachDylibNode *first;
+  SYMS_MachDylibNode *last;
+  SYMS_U32 count;
+} SYMS_MachDylibList;
+
+////////////////////////////////
+// Binds
+
+enum
+{
+  SYMS_MachBindTable_REGULAR,
+  SYMS_MachBindTable_LAZY,
+  SYMS_MachBindTable_WEAK
+};
+typedef SYMS_U32 SYMS_MachBindTable;
+
+typedef struct
+{
+  SYMS_U32 segment;
+  SYMS_U64 segment_offset;
+  SYMS_U64 dylib;
+  SYMS_String8 symbol_name;
+  SYMS_U8 flags;
+  SYMS_MachBindType type;
+  SYMS_S64 addend;
+} SYMS_MachBind;
+
+typedef struct SYMS_MachBindNode
+{
+  SYMS_MachBind data;
+  struct SYMS_MachBindNode *next;
+} SYMS_MachBindNode;
+
+typedef struct
+{
+  SYMS_MachBindNode *first;
+  SYMS_MachBindNode *last;
+  SYMS_U32 count;
+} SYMS_MachBindList;
+
+////////////////////////////////
+// Exports
+
+typedef struct SYMS_MachExport
+{
+  struct SYMS_MachExport **children;
+  SYMS_U8 child_count;
+  
+  SYMS_String8 name;
+  SYMS_U64 flags;
+  SYMS_U64 address;
+  
+  // SYMS_MachExportSymbolFlags_REEXPORT
+  SYMS_U64 dylib_ordinal;
+  SYMS_String8 import_name; 
+  
+  // SYMS_MachExportSymbolFlags_STUB_AND_RESOLVER
+  SYMS_U64 resolver;
+} SYMS_MachExport;
+
+typedef struct SYMS_MachExportNode
+{
+  struct SYMS_MachExportNode *next;
+  SYMS_MachExport data;
+} SYMS_MachExportNode;
+
+typedef struct SYMS_MachExportFrame
+{
+  struct SYMS_MachExportFrame *next;
+  SYMS_U8 child_idx;
+  SYMS_MachExport *node;
+} SYMS_MachExportFrame;
+
 ////////////////////////////////
 //~ NOTE(allen): MACH Parser Functions
 
@@ -71,4 +170,15 @@ SYMS_API SYMS_SecInfoArray      syms_mach_sec_info_array_from_bin(SYMS_Arena *ar
 
 SYMS_API SYMS_U64               syms_mach_default_vbase_from_bin(SYMS_MachBinAccel *bin);
 
+// dylib
+SYMS_API void syms_mach_dylib_list_push(SYMS_Arena *arena, SYMS_MachDylibList *list, SYMS_MachDylib *dylib, SYMS_U64Range name);
+
+// binds
+SYMS_API SYMS_MachBindList syms_mach_binds_from_base_range(SYMS_Arena *arena, void *base, SYMS_U64Range range, SYMS_U32 address_size, SYMS_MachBindTable bind_type);
+SYMS_API SYMS_ImportArray  syms_mach_imports_from_bin(SYMS_Arena *arena, SYMS_String8 data, SYMS_MachBinAccel *bin);
+
+// exports
+SYMS_API SYMS_MachExport * syms_build_mach_export_trie(SYMS_Arena *arena, void *base, SYMS_U64Range range);
+SYMS_API SYMS_ExportArray  syms_mach_exports_from_bin(SYMS_Arena *arena, SYMS_String8 data, SYMS_MachBinAccel *bin);
+
 #endif // SYMS_MACH_PARSER_H
diff --git a/Engine/Source/Runtime/SymsLib/syms/core/pdb/syms_pdb.h b/Engine/Source/Runtime/SymsLib/syms/core/pdb/syms_pdb.h
--- a/Engine/Source/Runtime/SymsLib/syms/core/pdb/syms_pdb.h
+++ b/Engine/Source/Runtime/SymsLib/syms/core/pdb/syms_pdb.h
@@ -160,12 +160,7 @@ typedef enum SYMS_PdbDbiRange{
   SYMS_PdbDbiRange_FileInfo,
   SYMS_PdbDbiRange_TSM,
   SYMS_PdbDbiRange_EcInfo,
-  
-  // TODO(allen): this name is weird; look into it. It really refers to a range
-  // of DBI that contains an array of section numbers (the PdbDbiDbgStream array).
-  // Correlate with the name dbg_header_size in SYMS_PdbDbiHeader.
   SYMS_PdbDbiRange_DbgHeader,
-  
   SYMS_PdbDbiRange_COUNT
 } SYMS_PdbDbiRange;
 
diff --git a/Engine/Source/Runtime/SymsLib/syms/core/pdb/syms_pdb_parser.c b/Engine/Source/Runtime/SymsLib/syms/core/pdb/syms_pdb_parser.c
--- a/Engine/Source/Runtime/SymsLib/syms/core/pdb/syms_pdb_parser.c
+++ b/Engine/Source/Runtime/SymsLib/syms/core/pdb/syms_pdb_parser.c
@@ -1021,8 +1021,10 @@ syms_pdb_unit_set_accel_from_dbg(SYMS_Arena *arena, SYMS_String8 data, SYMS_PdbD
           comp_unit.range_off[i] = cursor;
         }
         
-        // adjust symbols range off up by 4 to skip annoying signature
-        comp_unit.range_off[SYMS_PdbCompUnitRange_Symbols] += 4;
+        // skip annoying signature in first range
+        if (range_size[SYMS_PdbCompUnitRange_Symbols] > 4){
+          comp_unit.range_off[0] += 4;
+        }
         
         // fill comp unit direct fields
         comp_unit.sn = header.sn;
@@ -1693,6 +1695,16 @@ syms_pdb_sym_accel_from_range(SYMS_Arena *arena, SYMS_String8 data, SYMS_MsfAcce
   SYMS_PdbStubRef *var_last = 0;
   SYMS_U64 var_count = 0;
   
+  // thunk
+  SYMS_PdbStubRef *thunk_first = 0;
+  SYMS_PdbStubRef *thunk_last = 0;
+  SYMS_U64 thunk_count = 0;
+  
+  // pub
+  SYMS_PdbStubRef *pub_first = 0;
+  SYMS_PdbStubRef *pub_last = 0;
+  SYMS_U64 pub_count = 0;
+  
   // parse loop
   SYMS_PdbStub *defrange_collector_stub = 0;
   SYMS_PdbStubRef *stack = 0;
@@ -1743,6 +1755,23 @@ syms_pdb_sym_accel_from_range(SYMS_Arena *arena, SYMS_String8 data, SYMS_MsfAcce
         var_count += 1;
       }break;
       
+#if 0
+      case SYMS_CvSymKind_LTHREAD32:
+      case SYMS_CvSymKind_GTHREAD32:
+      {
+        SYMS_U32 name_off = sizeof(SYMS_CvThread32);
+        SYMS_String8 name = syms_msf_read_zstring_in_range(arena, data, msf, element.range, name_off);
+        
+        stub->name = name;
+        
+        // push onto var list
+        SYMS_PdbStubRef *ref = syms_pdb_alloc_ref(scratch.arena, &stack_free);
+        ref->stub = stub;
+        SYMS_QueuePush(var_first, var_last, ref);
+        var_count += 1;
+      }break;
+#endif
+      
       // general local variable
       case SYMS_CvSymKind_LOCAL:
       {
@@ -1846,9 +1875,27 @@ syms_pdb_sym_accel_from_range(SYMS_Arena *arena, SYMS_String8 data, SYMS_MsfAcce
       case SYMS_CvSymKind_PUB32:
       {
         SYMS_U32 name_off = sizeof(SYMS_CvPubsym32);
-        SYMS_String8 name = syms_msf_read_zstring_in_range(scratch.arena, data, msf, element.range, name_off);
+        SYMS_String8 name = syms_msf_read_zstring_in_range(arena, data, msf, element.range, name_off);
+        stub->name = name;
         
+        // push onto stripped symbol list
+        SYMS_PdbStubRef *ref = syms_pdb_alloc_ref(scratch.arena, &stack_free);
+        ref->stub = stub;
+        SYMS_QueuePush(pub_first, pub_last, ref);
+        pub_count += 1;
+      }break;
+      
+      case SYMS_CvSymKind_THUNK32:
+      {
+        SYMS_U32 name_off = sizeof(SYMS_CvThunk32);
+        SYMS_String8 name = syms_msf_read_zstring_in_range(scratch.arena, data, msf, element.range, name_off);
         stub->name = name;
+        
+        // push onto thunk list
+        SYMS_PdbStubRef *ref = syms_pdb_alloc_ref(scratch.arena, &stack_free);
+        ref->stub = stub;
+        SYMS_QueuePush(thunk_first, thunk_last, ref);
+        thunk_count += 1;
       }break;
     }
     
@@ -1902,6 +1949,28 @@ syms_pdb_sym_accel_from_range(SYMS_Arena *arena, SYMS_String8 data, SYMS_MsfAcce
     }
   }
   
+  // build thunk stubs pointer table
+  SYMS_PdbStub **thunk_stubs = syms_push_array(arena, SYMS_PdbStub*, thunk_count);
+  {
+    SYMS_PdbStub **ptr = thunk_stubs;
+    for (SYMS_PdbStubRef *ref = thunk_first;
+         ref != 0;
+         ref = ref->next, ptr += 1){
+      *ptr = ref->stub;
+    }
+  }
+  
+  // build pub stubs pointer table
+  SYMS_PdbStub **pub_stubs = syms_push_array(arena, SYMS_PdbStub*, pub_count);
+  {
+    SYMS_PdbStub **ptr = pub_stubs;
+    for (SYMS_PdbStubRef *ref = pub_first;
+         ref != 0;
+         ref = ref->next, ptr += 1){
+      *ptr = ref->stub;
+    }
+  }
+  
   // build bucket table
   SYMS_U64 bucket_count = (all_count/2)*2 + 3;
   SYMS_PdbStub **buckets = syms_push_array(arena, SYMS_PdbStub*, bucket_count);
@@ -1934,6 +2003,10 @@ syms_pdb_sym_accel_from_range(SYMS_Arena *arena, SYMS_String8 data, SYMS_MsfAcce
   result->proc_count = proc_count;
   result->var_stubs = var_stubs;
   result->var_count = var_count;
+  result->thunk_stubs = thunk_stubs;
+  result->thunk_count = thunk_count;
+  result->pub_stubs = pub_stubs;
+  result->pub_count = pub_count;
   
   return(result);
 }
@@ -1998,11 +2071,7 @@ syms_pdb_sym_symbol_kind_from_id(SYMS_String8 data, SYMS_PdbDbgAccel *dbg,
 SYMS_API SYMS_String8
 syms_pdb_sym_symbol_name_from_id(SYMS_Arena *arena, SYMS_String8 data, SYMS_PdbDbgAccel *dbg,
                                  SYMS_PdbUnitAccel *unit, SYMS_SymbolID id){
-  //- setup accel
-  SYMS_MsfAccel *msf = dbg->msf;
-  
   //- read id
-  SYMS_MsfRange range = syms_msf_range_from_sn(msf, unit->sn);
   SYMS_PdbStub *stub = syms_pdb_stub_from_unit_off(unit, SYMS_ID_u32_1(id));
   
   //- parse symbol
@@ -2822,7 +2891,7 @@ syms_pdb_leaf_type_info_from_id(SYMS_String8 data, SYMS_PdbDbgAccel *dbg, SYMS_P
         SYMS_CvLeafPointer ptr = {0};
         if (syms_msf_read_struct_in_range(data, msf, resolve.leaf_range, 0, &ptr)){
           SYMS_U64 size = SYMS_CvPointerAttribs_Extract_SIZE(ptr.attr);
-          SYMS_CvPointerKind ptr_kind = SYMS_CvPointerAttribs_Extract_KIND(ptr.attr);
+          //SYMS_CvPointerKind ptr_kind = SYMS_CvPointerAttribs_Extract_KIND(ptr.attr);
           SYMS_CvPointerMode ptr_mode = SYMS_CvPointerAttribs_Extract_MODE(ptr.attr);
           
           SYMS_TypeKind type_kind = syms_pdb_type_kind_from_cv_pointer_mode(ptr_mode);
@@ -3741,7 +3810,6 @@ syms_pdb_mem_info_from_number(SYMS_Arena *arena, SYMS_String8 data, SYMS_PdbDbgA
           result.kind = SYMS_MemKind_VBaseClassPtr;
           result.name = syms_push_string_copy(arena, stub->name);
           result.visibility = visibility;
-          //result.is_virtual = syms_true;
           result.off = stub->num;
         }break;
         
@@ -4107,6 +4175,85 @@ syms_pdb_scope_children_from_sid(SYMS_Arena *arena, SYMS_String8 data, SYMS_PdbD
   return(result);
 }
 
+SYMS_API void
+syms_pdb_stripped_sort_in_place(SYMS_StrippedInfo *info, SYMS_U64 count){
+  if (count > 1){
+    SYMS_U64 pivot = count - 1;
+    SYMS_U64 pivot_voff = info[pivot].voff;
+    
+    SYMS_B32 eq_flip_flop = syms_false;
+    
+    SYMS_U64 j = 0;
+    for (SYMS_U64 i = 0; i < pivot; i += 1){
+      SYMS_U64 voff = info[i].voff;
+      SYMS_B32 swap = syms_false;
+      if (voff < pivot_voff){
+        swap = syms_true;
+      }
+      else if (voff == pivot_voff){
+        swap = eq_flip_flop;
+        eq_flip_flop = !eq_flip_flop;
+      }
+      if (swap){
+        SYMS_Swap(SYMS_StrippedInfo, info[i], info[j]);
+        j += 1;
+      }
+    }
+    
+    SYMS_Swap(SYMS_StrippedInfo, info[j], info[pivot]);
+    
+    syms_pdb_stripped_sort_in_place(info, j);
+    syms_pdb_stripped_sort_in_place(info + j + 1, count - (j + 1));
+  }
+}
+
+SYMS_API SYMS_StrippedInfoArray
+syms_pdb_stripped_from_unit(SYMS_Arena *arena, SYMS_String8 data, SYMS_PdbDbgAccel *dbg, SYMS_PdbUnitAccel *unit){
+  SYMS_StrippedInfoArray result = {0};
+  
+  if (unit->pub_count != 0){
+    // setup accels
+    SYMS_MsfAccel *msf = dbg->msf;
+    
+    // allocate output memory
+    SYMS_U64 count = unit->pub_count;
+    SYMS_StrippedInfo *info = syms_push_array(arena, SYMS_StrippedInfo, count);
+    
+    // fill output array
+    SYMS_StrippedInfo *info_ptr = info;
+    SYMS_PdbStub **stub_opl = unit->pub_stubs + count;
+    for (SYMS_PdbStub **stub_ptr = unit->pub_stubs; stub_ptr < stub_opl; stub_ptr += 1, info_ptr += 1){
+      SYMS_PdbStub *stub = *stub_ptr;
+      
+      // extract info from stub
+      SYMS_String8 name = stub->name;
+      SYMS_U64 voff = 0;
+      {
+        SYMS_MsfRange range = syms_msf_range_from_sn(msf, unit->sn);
+        SYMS_CvElement sig_element = syms_cv_element(data, msf, range, stub->off);
+        SYMS_CvPubsym32 pubsym32 = {0};
+        if (syms_msf_read_struct_in_range(data, msf, sig_element.range, 0, &pubsym32)){
+          SYMS_CoffSection coff_section = syms_pdb_coff_section(data, dbg, pubsym32.sec);
+          voff = coff_section.virt_off + pubsym32.off;
+        }
+      }
+      
+      // fill stripped info
+      info_ptr->name = name;
+      info_ptr->voff = voff;
+    }
+    
+    // sort stripped info
+    syms_pdb_stripped_sort_in_place(info, count);
+    
+    // fill result struct
+    result.info = info;
+    result.count = count;
+  }
+  
+  return(result);
+}
+
 
 ////////////////////////////////
 //~ NOTE(allen): PDB Signature Info
diff --git a/Engine/Source/Runtime/SymsLib/syms/core/pdb/syms_pdb_parser.h b/Engine/Source/Runtime/SymsLib/syms/core/pdb/syms_pdb_parser.h
--- a/Engine/Source/Runtime/SymsLib/syms/core/pdb/syms_pdb_parser.h
+++ b/Engine/Source/Runtime/SymsLib/syms/core/pdb/syms_pdb_parser.h
@@ -152,6 +152,10 @@ typedef struct SYMS_PdbUnitAccel{
   SYMS_U64 proc_count;
   SYMS_PdbStub **var_stubs;
   SYMS_U64 var_count;
+  SYMS_PdbStub **thunk_stubs;
+  SYMS_U64 thunk_count;
+  SYMS_PdbStub **pub_stubs;
+  SYMS_U64 pub_count;
   SYMS_PdbStub **udt_stubs;
   SYMS_U64 udt_count;
 } SYMS_PdbUnitAccel;
@@ -528,6 +532,11 @@ SYMS_API SYMS_SymbolIDArray syms_pdb_scope_children_from_sid(SYMS_Arena *arena,
                                                              SYMS_PdbDbgAccel *dbg, SYMS_PdbUnitAccel *unit,
                                                              SYMS_SymbolID id);
 
+SYMS_API void                   syms_pdb_stripped_sort_in_place(SYMS_StrippedInfo *info, SYMS_U64 count);
+
+SYMS_API SYMS_StrippedInfoArray syms_pdb_stripped_from_unit(SYMS_Arena *arena, SYMS_String8 data,
+                                                            SYMS_PdbDbgAccel *dbg, SYMS_PdbUnitAccel *unit);
+
 ////////////////////////////////
 //~ NOTE(allen): PDB Signature Info
 
diff --git a/Engine/Source/Runtime/SymsLib/syms/core/pe/syms_pe.h b/Engine/Source/Runtime/SymsLib/syms/core/pe/syms_pe.h
--- a/Engine/Source/Runtime/SymsLib/syms/core/pe/syms_pe.h
+++ b/Engine/Source/Runtime/SymsLib/syms/core/pe/syms_pe.h
@@ -97,5 +97,34 @@ typedef struct SYMS_PeCvHeaderPDB70{
 } SYMS_PeCvHeaderPDB70;
 #pragma pack(pop)
 
+////////////////////////////////
+// Import
+
+typedef struct
+{
+  SYMS_U32 lookup_table_virt_off;
+  SYMS_U32 timestamp;
+  SYMS_U32 forwarder_chain;
+  SYMS_U32 name_virt_off;
+  SYMS_U32 import_addr_table_virt_off;
+} SYMS_PeImportDirectoryEntry;
+
+////////////////////////////////
+// Export
+
+typedef struct
+{
+  SYMS_U32 flags;                       // must be zero
+  SYMS_U32 timestamp;                   // time and date when export table was created
+  SYMS_U16 major_ver;                   // table version, user can change major and minor version
+  SYMS_U16 minor_ver; 
+  SYMS_U32 name_virt_off;               // ASCII name of the dll
+  SYMS_U32 ordinal_base;                // Starting oridnal number
+  SYMS_U32 export_address_table_count;
+  SYMS_U32 name_pointer_table_count;
+  SYMS_U32 export_address_table_virt_off;
+  SYMS_U32 name_pointer_table_virt_off;
+  SYMS_U32 ordinal_table_virt_off;
+} SYMS_PeExportTable;
 
 #endif //SYMS_PE_H
