Update SymsLib to latest version and update SymsResolver for Insights.

diff --git a/Engine/Source/Runtime/SymsLib/syms/core/pe/syms_pe_parser.c b/Engine/Source/Runtime/SymsLib/syms/core/pe/syms_pe_parser.c
--- a/Engine/Source/Runtime/SymsLib/syms/core/pe/syms_pe_parser.c
+++ b/Engine/Source/Runtime/SymsLib/syms/core/pe/syms_pe_parser.c
@@ -62,8 +62,9 @@ syms_pe_bin_accel_from_file(SYMS_Arena *arena, SYMS_String8 data, SYMS_PeFileAcc
     SYMS_CoffSection *sections = (SYMS_CoffSection*)((SYMS_U8*)base + sec_array_off);
     
     // get data directory from optional
-    SYMS_U16 data_dir_offset = 0;
     SYMS_U32 data_dir_count = 0;
+    SYMS_U64Range *data_dirs_file = 0;
+    SYMS_U64Range *data_dirs_virt = 0;
     if (optional_size > 0){
       // read magic number
       SYMS_U16 optional_magic = 0;
@@ -89,10 +90,19 @@ syms_pe_bin_accel_from_file(SYMS_Arena *arena, SYMS_String8 data, SYMS_PeFileAcc
         }break;
       }
       
-      // fill outputs from optional
-      SYMS_U32 data_dir_max = (optional_size - data_dir_offset)/sizeof(SYMS_PeDataDirectory);
+      SYMS_U32 data_dir_max = (optional_size - reported_data_dir_offset)/sizeof(SYMS_PeDataDirectory);
       data_dir_count = SYMS_ClampTop(reported_data_dir_count, data_dir_max);
-      data_dir_offset = optional_range.min + reported_data_dir_offset;
+      // convert PE directories to ranges
+      data_dirs_virt = syms_push_array(arena, SYMS_U64Range, data_dir_count);
+      data_dirs_file = syms_push_array(arena, SYMS_U64Range, data_dir_count);
+      for (SYMS_U32 dir_idx = 0; dir_idx < data_dir_count; dir_idx += 1){
+        SYMS_U64 dir_offset = optional_range.min + reported_data_dir_offset + sizeof(SYMS_PeDataDirectory)*dir_idx;
+        SYMS_PeDataDirectory dir; syms_memzero_struct(&dir);
+        syms_based_range_read_struct(base, range, dir_offset, &dir);
+        SYMS_U64 file_off = syms_pe_virt_off_to_file_off(sections, clamped_sec_count, dir.virt_off);
+        data_dirs_virt[dir_idx] = syms_make_u64_range(dir.virt_off, dir.virt_off + dir.virt_size);
+        data_dirs_file[dir_idx] = syms_make_u64_inrange(range, file_off, dir.virt_size);
+      }
     }
     
     // read info about debug file
@@ -103,59 +113,37 @@ syms_pe_bin_accel_from_file(SYMS_Arena *arena, SYMS_String8 data, SYMS_PeFileAcc
     SYMS_U64 dbg_path_size = 0;
     
     if (SYMS_PeDataDirectoryIndex_DEBUG < data_dir_count){
-      SYMS_PeDataDirectory dbg_data_dir = {0};
-      SYMS_U64 dbg_data_dir_offset = data_dir_offset + sizeof(SYMS_PeDataDirectory)*SYMS_PeDataDirectoryIndex_DEBUG;
-      syms_based_range_read_struct(base, range, dbg_data_dir_offset, &dbg_data_dir);
-      
-      // find section that virtually contains the debug data
-      SYMS_CoffSection *dbg_data_sec = 0;
-      SYMS_CoffSection *sec = sections;
-      for (SYMS_U64 i = 0; i < clamped_sec_count; i += 1, sec += 1){
-        SYMS_U64 first = sec->virt_off;
-        SYMS_U64 opl   = first + sec->virt_size;
-        if (first <= dbg_data_dir.virt_off && dbg_data_dir.virt_off < opl){
-          dbg_data_sec = sec;
-          break;
-        }
-      }
-      
       // read debug directory
-      if (dbg_data_sec != 0){
-        SYMS_U64 dbg_data_offset_from_sec = dbg_data_dir.virt_off - dbg_data_sec->virt_off;
-        SYMS_U64 dbg_data_offset = dbg_data_sec->file_off + dbg_data_offset_from_sec;
-        SYMS_PeDebugDirectory dbg_data = {0};
-        SYMS_U64 dbg_data_size = SYMS_ClampTop(dbg_data_dir.virt_size, sizeof(dbg_data));
-        syms_based_range_read(base, range, dbg_data_offset, dbg_data_size, &dbg_data);
-        
-        // extract external file info from codeview header
-        if (dbg_data.type == SYMS_PeDebugDirectoryType_CODEVIEW){
-          SYMS_U64 cv_offset = dbg_data.file_offset;
-          SYMS_U32 cv_magic = 0;
-          syms_based_range_read_struct(base, range, cv_offset, &cv_magic);
-          switch (cv_magic){
-            default:break;
-            case SYMS_CODEVIEW_PDB20_MAGIC:
-            {
-              SYMS_PeCvHeaderPDB20 cv = {0};
-              syms_based_range_read_struct(base, range, cv_offset, &cv);
-              // TODO(allen): can we extend the ext match key system to use this in some way?
-              dbg_time = cv.time;
-              dbg_age = cv.age;
-              dbg_path_off = cv_offset + sizeof(cv);
-            }break;
-            case SYMS_CODEVIEW_PDB70_MAGIC:
-            {
-              SYMS_PeCvHeaderPDB70 cv = {0};
-              syms_based_range_read_struct(base, range, cv_offset, &cv);
-              dbg_guid = cv.guid;
-              dbg_age = cv.age;
-              dbg_path_off = cv_offset + sizeof(cv);
-            }break;
-          }
-          if (dbg_path_off > 0){
-            SYMS_String8 dbg_path = syms_based_range_read_string(base, range, dbg_path_off);
-            dbg_path_size = dbg_path.size;
-          }
+      SYMS_PeDebugDirectory dbg_data = {0};
+      syms_based_range_read_struct(base, data_dirs_file[SYMS_PeDataDirectoryIndex_DEBUG], 0, &dbg_data);
+      // extract external file info from codeview header
+      if (dbg_data.type == SYMS_PeDebugDirectoryType_CODEVIEW){
+        SYMS_U64 cv_offset = dbg_data.file_offset;
+        SYMS_U32 cv_magic = 0;
+        syms_based_range_read_struct(base, range, cv_offset, &cv_magic);
+        switch (cv_magic){
+          default:break;
+          case SYMS_CODEVIEW_PDB20_MAGIC:
+          {
+            SYMS_PeCvHeaderPDB20 cv = {0};
+            syms_based_range_read_struct(base, range, cv_offset, &cv);
+            // TODO(allen): can we extend the ext match key system to use this in some way?
+            dbg_time = cv.time;
+            dbg_age = cv.age;
+            dbg_path_off = cv_offset + sizeof(cv);
+          }break;
+          case SYMS_CODEVIEW_PDB70_MAGIC:
+          {
+            SYMS_PeCvHeaderPDB70 cv = {0};
+            syms_based_range_read_struct(base, range, cv_offset, &cv);
+            dbg_guid = cv.guid;
+            dbg_age = cv.age;
+            dbg_path_off = cv_offset + sizeof(cv);
+          }break;
+        }
+        if (dbg_path_off > 0){
+          SYMS_String8 dbg_path = syms_based_range_read_string(base, range, dbg_path_off);
+          dbg_path_size = dbg_path.size;
         }
       }
     }
@@ -168,65 +156,12 @@ syms_pe_bin_accel_from_file(SYMS_Arena *arena, SYMS_String8 data, SYMS_PeFileAcc
     result->dbg_path_off = dbg_path_off;
     result->dbg_path_size = dbg_path_size;
     result->dbg_guid = dbg_guid;
+    result->dbg_age = dbg_age;
     result->arch = syms_arch_from_coff_machine_type(coff_header.machine);
-    
-    
-    // TODO(allen): HACK HACK HACK HACK HACK HACK HACK HACK HACK HACK
-    // This only works because we use it after initializing the rest of PeBinAccel.
-    // The good news is we don't really want to bake this in here long term anyways
-    // so the goal is to get this moved out into a separate UnwindAccel parsing stage.
-    
-    // TODO(allen): only applies for "intel" architectures - need a handler for
-    // arm and mips data structures too.
-    
-    // read pdata
-    SYMS_U64 pdata_off = 0;
-    SYMS_U64 pdata_count = 0;
-    
-    {
-      // try to get pdata section
-      SYMS_CoffSection *pdata_sec = 0;
-      SYMS_CoffSection *coff_sec = (SYMS_CoffSection*)(data.str + sec_array_off);
-      for (SYMS_U64 n = 1; n <= clamped_sec_count; n += 1, coff_sec += 1){
-        
-        // extract name
-        SYMS_U8 *name_base = coff_sec->name;
-        SYMS_U8 *name_ptr = name_base;
-        SYMS_U8 *name_opl = name_base + 8;
-        for (;name_ptr < name_opl && *name_ptr != 0; name_ptr += 1);
-        SYMS_String8 sec_name = syms_str8_range(name_base, name_ptr);
-        
-        // check name
-        if (syms_string_match(sec_name, syms_str8_lit(".pdata"), 0)){
-          pdata_sec = coff_sec;
-          break;
-        }
-      }
-      
-      // determine pdata offset & count
-      if (pdata_sec != 0){
-        SYMS_U64 off = pdata_sec->file_off;
-        SYMS_U64 size = pdata_sec->file_size;
-        
-        // save offset to pdata
-        pdata_off = off;
-        
-        // scan range of pdata to eliminate nulls at the end
-        SYMS_PeIntelPdata *first_pdata_intel = (SYMS_PeIntelPdata*)(data.str + off);
-        SYMS_PeIntelPdata *opl_pdata_intel = first_pdata_intel + (size/sizeof(SYMS_PeIntelPdata));
-        if (first_pdata_intel < opl_pdata_intel){
-          SYMS_PeIntelPdata *pdata_intel_ptr = opl_pdata_intel - 1;
-          for (;pdata_intel_ptr >= first_pdata_intel && pdata_intel_ptr->voff_first == 0;
-               pdata_intel_ptr -= 1);
-          pdata_intel_ptr += 1;
-          pdata_count = (SYMS_U64)(pdata_intel_ptr - first_pdata_intel);
-        }
-      }
-    }
-    
-    // fill pdata
-    result->pdata_off = pdata_off;
-    result->pdata_count = pdata_count;
+    // store virtual and file ranges for data directories
+    result->data_dir_count = data_dir_count;
+    result->data_dirs_virt = data_dirs_virt;
+    result->data_dirs_file = data_dirs_file;
   }
   
   return(result);
@@ -296,19 +231,22 @@ syms_pe_sec_info_array_from_bin(SYMS_Arena *arena, SYMS_String8 data, SYMS_PeBin
 
 SYMS_API SYMS_U64
 syms_pe_binary_search_intel_pdata(SYMS_String8 data, SYMS_PeBinAccel *bin, SYMS_U64 voff){
+  SYMS_ASSERT(bin->arch == SYMS_Arch_X86 || bin->arch == SYMS_Arch_X64);
+  SYMS_U64 pdata_off = bin->data_dirs_file[SYMS_PeDataDirectoryIndex_EXCEPTIONS].min;
+  SYMS_U64 pdata_count = syms_u64_range_size(bin->data_dirs_file[SYMS_PeDataDirectoryIndex_EXCEPTIONS]) / sizeof(SYMS_PeIntelPdata);
+
   SYMS_U64 result = 0;
-  
   // check if this bin includes a pdata array
-  if (bin->pdata_count > 0){
-    SYMS_PeIntelPdata *pdata_array = (SYMS_PeIntelPdata*)(data.str + bin->pdata_off);
+  if (pdata_count > 0){
+    SYMS_PeIntelPdata *pdata_array = (SYMS_PeIntelPdata*)(data.str + pdata_off);
     if (voff >= pdata_array[0].voff_first){
       
       // binary search:
       //  find max index s.t. pdata_array[index].voff_first <= voff
       //  we assume (i < j) -> (pdata_array[i].voff_first < pdata_array[j].voff_first)
-      SYMS_U64 index = bin->pdata_count;
+      SYMS_U64 index = pdata_count;
       SYMS_U64 min = 0;
-      SYMS_U64 opl = bin->pdata_count;
+      SYMS_U64 opl = pdata_count;
       for (;;){
         SYMS_U64 mid = (min + opl)/2;
         SYMS_PeIntelPdata *pdata = pdata_array + mid;
@@ -332,7 +270,7 @@ syms_pe_binary_search_intel_pdata(SYMS_String8 data, SYMS_PeBinAccel *bin, SYMS_
       {
         SYMS_PeIntelPdata *pdata = pdata_array + index;
         if (pdata->voff_first <= voff && voff < pdata->voff_one_past_last){
-          result = bin->pdata_off + index*sizeof(SYMS_PeIntelPdata);
+          result = pdata_off + index*sizeof(SYMS_PeIntelPdata);
         }
       }
     }
@@ -402,4 +340,220 @@ syms_pe_ptr_from_voff(SYMS_String8 data, SYMS_PeBinAccel *bin, SYMS_U64 voff){
   return(result);
 }
 
+SYMS_API SYMS_U64
+syms_pe_virt_off_to_file_off(SYMS_CoffSection *sections, SYMS_U32 section_count, SYMS_U64 virt_off)
+{
+  SYMS_U64 file_off = 0;
+  for (SYMS_U32 sect_idx = 0; sect_idx < section_count; sect_idx += 1){
+    SYMS_CoffSection *sect = &sections[sect_idx];
+    if (sect->virt_off <= virt_off && virt_off < sect->virt_off + sect->virt_size){
+      file_off = sect->file_off + (virt_off - sect->virt_off);
+    }
+  }
+  return file_off;
+}
+
+SYMS_API SYMS_U64
+syms_pe_bin_virt_off_to_file_off(SYMS_String8 data, SYMS_PeBinAccel *bin, SYMS_U64 virt_off)
+{
+  return syms_pe_virt_off_to_file_off((SYMS_CoffSection*)((SYMS_U8*)data.str + bin->section_array_off), bin->section_count, virt_off);
+}
+
+SYMS_API SYMS_ImportArray
+syms_pe_imports_from_bin(SYMS_Arena *arena, SYMS_String8 data, SYMS_PeBinAccel *bin)
+{
+  SYMS_ArenaTemp scratch = syms_get_scratch(&arena, 1);
+  
+  //- rjf: grab prerequisites
+  void *base = (void*)data.str;
+  SYMS_Arch arch = syms_pe_arch_from_bin(bin);
+  SYMS_U64 arch_bit_size = syms_address_size_from_arch(bin->arch);
+  SYMS_U64 arch_byte_size = arch_bit_size / 8;
+  
+  //- rjf: grab ranges
+  SYMS_U64Range dir_range = bin->data_dirs_file[SYMS_PeDataDirectoryIndex_IMPORT];
+  SYMS_U64Range data_range = syms_make_u64_range(0, data.size);
+  
+  //- rjf: parse import list
+  SYMS_PeImportDllNode *first_dll = 0;
+  SYMS_PeImportDllNode *last_dll = 0;
+  SYMS_U32 dll_count = 0;
+  {
+    for(SYMS_U64 read_offset = 0;;)
+    {
+      // rjf: read entry
+      SYMS_PeImportDirectoryEntry import_entry = {0};
+      read_offset += syms_based_range_read_struct(base, dir_range, read_offset, &import_entry);
+      
+      // rjf: break if this is the last entry
+      SYMS_B32 is_last_import_directory = (import_entry.lookup_table_virt_off == 0 &&
+                                           import_entry.timestamp == 0 &&
+                                           import_entry.forwarder_chain == 0 &&
+                                           import_entry.name_virt_off == 0 &&
+                                           import_entry.import_addr_table_virt_off == 0);
+      if(is_last_import_directory)
+      {
+        break;
+      }
+      
+      // rjf: grab offsets
+      SYMS_U64 roff__dll_name_file = syms_pe_bin_virt_off_to_file_off(data, bin, import_entry.name_virt_off);
+      SYMS_U64 roff__lookup_table = syms_pe_bin_virt_off_to_file_off(data, bin, import_entry.lookup_table_virt_off);
+      
+      // rjf: build import DLL node
+      SYMS_PeImportDllNode *dll_node = syms_push_array(scratch.arena, SYMS_PeImportDllNode, 1);
+      dll_node->next = 0;
+      dll_node->import_entry = import_entry;
+      dll_node->name = syms_based_range_read_string(base, data_range, roff__dll_name_file);
+      dll_node->first_import = 0;
+      dll_node->last_import = 0;
+      dll_node->import_count = 0;
+      dll_count += 1;
+      SYMS_QueuePush(first_dll, last_dll, dll_node);
+      
+      // rjf: build list of all imports for this DLL
+      for(;;)
+      {
+        // rjf: parse entry
+        SYMS_U64 lookup_entry = 0;
+        roff__lookup_table += syms_based_range_read(base, data_range, roff__lookup_table, arch_byte_size, &lookup_entry);
+        
+        // rjf: break if we're done
+        if(lookup_entry == 0)
+        {
+          break;
+        }
+        
+        // rjf: build node
+        SYMS_PeImportNode *import_node = syms_push_array_zero(scratch.arena, SYMS_PeImportNode, 1);
+        SYMS_QueuePush(dll_node->first_import, dll_node->last_import, import_node);
+        dll_node->import_count += 1;
+        
+        // rjf: grab name/ordinal
+        SYMS_B32 is_ordinal = (arch_bit_size == 64) ? lookup_entry & 0x8000000000000000 : lookup_entry & 0x80000000;
+        if(is_ordinal)
+        {
+          import_node->ordinal = lookup_entry & 0xFFFF;
+        }
+        else
+        {
+          SYMS_U64 hn_file_off = syms_pe_bin_virt_off_to_file_off(data, bin, lookup_entry);
+          syms_based_range_read_struct(base, data_range, hn_file_off, &import_node->hint);
+          import_node->name = syms_based_range_read_string(base, data_range, hn_file_off + sizeof(import_node->hint));
+        }
+      }
+      
+    }
+  }
+  
+  //- rjf: sum imports
+  SYMS_U64 total_import_count = 0;
+  for(SYMS_PeImportDllNode *n = first_dll; n != 0; n = n->next)
+  {
+    total_import_count += n->import_count;
+  }
+  
+  //- rjf: build+fill result
+  SYMS_ImportArray import_array = {0};
+  {
+    SYMS_U64 import_write_idx = 0;
+    import_array.count = total_import_count;
+    import_array.imports = syms_push_array_zero(arena, SYMS_Import, import_array.count);
+    for(SYMS_PeImportDllNode *dll_node = first_dll;
+        dll_node != 0;
+        dll_node = dll_node->next)
+    {
+      for(SYMS_PeImportNode *import_node = dll_node->first_import;
+          import_node != 0;
+          import_node = import_node->next)
+      {
+        SYMS_Import *imp = &import_array.imports[import_write_idx];
+        imp->library_name = syms_push_string_copy(arena, dll_node->name);
+        imp->name = syms_push_string_copy(arena, import_node->name);
+        imp->ordinal = import_node->ordinal;
+        import_write_idx += 1;
+      }
+    }
+  }
+  
+  //- rjf: return
+  syms_release_scratch(scratch);
+  return import_array;
+}
+
+SYMS_API SYMS_ExportArray
+syms_pe_exports_from_bin(SYMS_Arena *arena, SYMS_String8 data, SYMS_PeBinAccel *bin)
+{
+  SYMS_U64Range dir_range = bin->data_dirs_file[SYMS_PeDataDirectoryIndex_EXPORT];
+  SYMS_U64Range data_range = syms_make_u64_range(0, data.size);
+  void *base = (void*)data.str;
+  
+  SYMS_PeExportTable header; syms_memzero_struct(&header);
+  syms_based_range_read_struct(base, dir_range, 0, &header);
+  
+  SYMS_U32 export_address_table_file_off = syms_pe_bin_virt_off_to_file_off(data, bin, header.export_address_table_virt_off);
+  SYMS_U32 name_pointer_table_file_off   = syms_pe_bin_virt_off_to_file_off(data, bin, header.name_pointer_table_virt_off);
+  SYMS_U32 ordinal_table_file_off        = syms_pe_bin_virt_off_to_file_off(data, bin, header.ordinal_table_virt_off);
+  
+  SYMS_U64Range export_table_range       = syms_make_u64_inrange(data_range, export_address_table_file_off, sizeof(SYMS_U32) * header.export_address_table_count);
+  SYMS_U64Range name_pointer_table_range = syms_make_u64_inrange(data_range, name_pointer_table_file_off,   sizeof(SYMS_U32) * header.name_pointer_table_count);
+  SYMS_U64Range ordinal_table_range      = syms_make_u64_inrange(data_range, ordinal_table_file_off,        sizeof(SYMS_U16) * header.name_pointer_table_count);
+  
+  // NOTE(nick): There is a DLL name string in the export header. This is
+  // intersting because if we read image from memory we can figure out its
+  // name by reading this string.
+#if 0
+  SYMS_U32 dll_name_file_off = syms_pe_bin_virt_off_to_file_off(data, bin, header.name_virt_off);
+  SYMS_String8 dll_name = syms_based_range_read_string(base, data_range, dll_name_file_off);
+#endif
+  
+  SYMS_ExportArray export_array;
+  export_array.count = header.name_pointer_table_count;
+  export_array.exports = syms_push_array_zero(arena, SYMS_Export, export_array.count);
+  
+  for (SYMS_U32 i = 0; i < header.name_pointer_table_count; i += 1) {
+    SYMS_U32 name_virt_off = 0;
+    syms_based_range_read_struct(base, name_pointer_table_range, i * sizeof(name_virt_off), &name_virt_off);
+    SYMS_U32 name_file_off = syms_pe_bin_virt_off_to_file_off(data, bin, name_virt_off);
+    SYMS_String8 name = syms_based_range_read_string(base, data_range, name_file_off);
+    
+    SYMS_U16 ordinal = 0;
+    syms_based_range_read_struct(base, ordinal_table_range, i * sizeof(ordinal), &ordinal);
+    SYMS_U32 biased_ordinal = ordinal + header.ordinal_base;
+    
+    SYMS_U32 export_virt_off = 0;
+    syms_based_range_read_struct(base, export_table_range, ordinal * sizeof(export_virt_off), &export_virt_off);
+    SYMS_U32 export_file_off = syms_pe_bin_virt_off_to_file_off(data, bin, export_virt_off);
+    
+    SYMS_Export *exp = &export_array.exports[i];
+    SYMS_B32 is_virt_off_forwarder = bin->data_dirs_virt[SYMS_PeDataDirectoryIndex_EXPORT].min <= export_virt_off && export_virt_off < bin->data_dirs_virt[SYMS_PeDataDirectoryIndex_EXPORT].max;
+    if (is_virt_off_forwarder) {
+      SYMS_String8 library_name = syms_str8_lit("@MALFORMED_LIBRARY_NAME@");
+      SYMS_String8 import_name = syms_str8_lit("@MALFORMED_IMPORT_NAME@");
+      // parse forwarder command string that is formatted like "<DLL>.<IMPORT>"
+      SYMS_String8 forwarder = syms_based_range_read_string(base, data_range, export_file_off);
+      SYMS_ArenaTemp temp = syms_arena_temp_begin(arena);
+      SYMS_String8List list = syms_string_split(temp.arena, forwarder, '.');
+      if (list.node_count == 2) {
+        library_name = list.first->string;
+        import_name  = list.last->string;
+      }
+      syms_arena_temp_end(temp);
+      // fill out export
+      exp->name = syms_push_string_copy(arena, name);
+      exp->address = 0;
+      exp->ordinal = biased_ordinal;
+      exp->forwarder_library_name = syms_push_string_copy(arena, library_name);
+      exp->forwarder_import_name = syms_push_string_copy(arena, import_name);
+    } else {
+      // fill out export
+      exp->name = syms_push_string_copy(arena, name);
+      exp->address = export_virt_off;
+      exp->ordinal = biased_ordinal;
+    }
+  }
+  
+  return export_array;
+}
+
 #endif //SYMS_PE_PARSER_C
diff --git a/Engine/Source/Runtime/SymsLib/syms/core/pe/syms_pe_parser.h b/Engine/Source/Runtime/SymsLib/syms/core/pe/syms_pe_parser.h
--- a/Engine/Source/Runtime/SymsLib/syms/core/pe/syms_pe_parser.h
+++ b/Engine/Source/Runtime/SymsLib/syms/core/pe/syms_pe_parser.h
@@ -19,13 +19,42 @@ typedef struct SYMS_PeBinAccel{
   SYMS_U64 dbg_path_off;
   SYMS_U64 dbg_path_size;
   SYMS_PeGuid dbg_guid;
+  SYMS_U32 dbg_age;
   SYMS_Arch arch;
-  
-  // TODO(allen): refine this part
-  SYMS_U64 pdata_off;
-  SYMS_U64 pdata_count;
+  SYMS_U32 data_dir_count;
+  SYMS_U64Range *data_dirs_virt;
+  SYMS_U64Range *data_dirs_file;
 } SYMS_PeBinAccel;
 
+////////////////////////////////
+// Imports
+
+enum
+{
+  SYMS_PeLookupBy_Ordinal,
+  SYMS_PeLookupBy_NameHint
+};
+typedef SYMS_U32 SYMS_PeImportLookupBy;
+
+typedef struct SYMS_PeImportNode
+{
+  struct SYMS_PeImportNode *next;
+  SYMS_PeImportLookupBy lookup_by;
+  SYMS_U16 ordinal;
+  SYMS_String8 name;
+  SYMS_U16 hint;
+} SYMS_PeImportNode;
+
+typedef struct SYMS_PeImportDllNode
+{
+  struct SYMS_PeImportDllNode *next;
+  SYMS_String8 name;
+  SYMS_PeImportDirectoryEntry import_entry;
+  SYMS_PeImportNode *first_import;
+  SYMS_PeImportNode *last_import;
+  SYMS_U32 import_count;
+} SYMS_PeImportDllNode;
+
 ////////////////////////////////
 // NOTE(allen): PE Parser Functions
 
@@ -60,4 +89,12 @@ SYMS_API void*    syms_pe_ptr_from_sec_number(SYMS_String8 data, SYMS_PeBinAccel
 SYMS_API void*    syms_pe_ptr_from_foff(SYMS_String8 data, SYMS_PeBinAccel *bin, SYMS_U64 foff);
 SYMS_API void*    syms_pe_ptr_from_voff(SYMS_String8 data, SYMS_PeBinAccel *bin, SYMS_U64 voff);
 
+SYMS_API SYMS_U64 syms_pe_virt_off_to_file_off(SYMS_CoffSection *sections, SYMS_U32 section_count, SYMS_U64 virt_off);
+SYMS_API SYMS_U64 syms_pe_bin_virt_off_to_file_off(SYMS_String8 data, SYMS_PeBinAccel *bin, SYMS_U64 virt_off);
+
+////////////////////////////////
+
+SYMS_API SYMS_ImportArray syms_pe_imports_from_bin(SYMS_Arena *arena, SYMS_String8 data, SYMS_PeBinAccel *bin);
+SYMS_API SYMS_ExportArray syms_pe_exports_from_bin(SYMS_Arena *arena, SYMS_String8 data, SYMS_PeBinAccel *bin);
+
 #endif //SYMS_PE_PARSER_H
diff --git a/Engine/Source/Runtime/SymsLib/syms/core/regs/syms_regs.c b/Engine/Source/Runtime/SymsLib/syms/core/regs/syms_regs.c
new file mode 100644
--- /dev/null
+++ b/Engine/Source/Runtime/SymsLib/syms/core/regs/syms_regs.c
@@ -0,0 +1,19 @@
+// Copyright Epic Games, Inc. All Rights Reserved.
+#ifndef SYMS_REGS_C
+#define SYMS_REGS_C
+
+////////////////////////////////
+// NOTE(allen): Basic Register Helpers
+
+SYMS_API SYMS_String8
+syms_reg_from_metadata_id(SYMS_RegMetadata *metadata, SYMS_U64 count, void *reg_file, SYMS_RegID reg_id){
+  SYMS_String8 result = {0};
+  if (reg_id <= count){
+    SYMS_RegMetadata *m = metadata + reg_id;
+    result.str = (SYMS_U8*)reg_file + m->offset;
+    result.size = m->size;
+  }
+  return(result);
+}
+
+#endif // SYMS_REGS_C
diff --git a/Engine/Source/Runtime/SymsLib/syms/core/regs/syms_regs.h b/Engine/Source/Runtime/SymsLib/syms/core/regs/syms_regs.h
--- a/Engine/Source/Runtime/SymsLib/syms/core/regs/syms_regs.h
+++ b/Engine/Source/Runtime/SymsLib/syms/core/regs/syms_regs.h
@@ -58,4 +58,11 @@ typedef struct SYMS_RegMetadata{
   SYMS_U32 alias_to;
 } SYMS_RegMetadata;
 
+////////////////////////////////
+// NOTE(allen): Basic Register Helpers
+
+SYMS_API SYMS_String8 syms_reg_from_metadata_id(SYMS_RegMetadata *metadata, SYMS_U64 count,
+                                                void *reg_file, SYMS_RegID reg_id);
+#define syms_reg_from_arch_id(AR,f,i) syms_reg_from_metadata_id(reg_metadata_##AR,SYMS_Reg##AR##Code_COUNT,(f),(i))
+
 #endif // SYMS_REGS_H
diff --git a/Engine/Source/Runtime/SymsLib/syms/core/syms_debug_info.h b/Engine/Source/Runtime/SymsLib/syms/core/syms_debug_info.h
--- a/Engine/Source/Runtime/SymsLib/syms/core/syms_debug_info.h
+++ b/Engine/Source/Runtime/SymsLib/syms/core/syms_debug_info.h
@@ -33,6 +33,44 @@ typedef struct SYMS_BinInfoArray{
   SYMS_U64 count;
 } SYMS_BinInfoArray;
 
+////////////////////////////////
+//~ NOTE(allen): Binary Imports & Exports
+
+typedef struct SYMS_Import{
+  SYMS_String8 name;
+  SYMS_String8 library_name;
+  SYMS_U16     ordinal;
+} SYMS_Import;
+
+typedef struct SYMS_Export{
+  SYMS_String8 name;
+  SYMS_U64 address;
+  SYMS_U32 ordinal;
+  SYMS_U32 unused_;
+  SYMS_String8 forwarder_library_name;
+  SYMS_String8 forwarder_import_name;
+} SYMS_Export;
+
+typedef struct SYMS_ImportNode{
+  struct SYMS_ImportNode *next;
+  SYMS_Import data;
+} SYMS_ImportNode;
+
+typedef struct SYMS_ExportNode{
+  struct SYMS_ExportNode *next;
+  SYMS_Export data;
+} SYMS_ExportNode;
+
+typedef struct SYMS_ImportArray{
+  SYMS_Import *imports;
+  SYMS_U64 count;
+} SYMS_ImportArray;
+
+typedef struct SYMS_ExportArray{
+  SYMS_Export *exports;
+  SYMS_U64 count;
+} SYMS_ExportArray;
+
 ////////////////////////////////
 //~ NOTE(allen): Compilation Unit Types
 
@@ -254,6 +292,16 @@ typedef struct SYMS_ConstInfo{
   SYMS_U64 val;
 } SYMS_ConstInfo;
 
+typedef struct SYMS_StrippedInfo{
+  SYMS_String8 name;
+  SYMS_U64 voff;
+} SYMS_StrippedInfo;
+
+typedef struct SYMS_StrippedInfoArray{
+  SYMS_StrippedInfo *info;
+  SYMS_U64 count;
+} SYMS_StrippedInfoArray;
+
 typedef enum SYMS_MemKind{
   SYMS_MemKind_Null,
   SYMS_MemKind_DataField,
diff --git a/Engine/Source/Runtime/SymsLib/syms/core/syms_parser.c b/Engine/Source/Runtime/SymsLib/syms/core/syms_parser.c
--- a/Engine/Source/Runtime/SymsLib/syms/core/syms_parser.c
+++ b/Engine/Source/Runtime/SymsLib/syms/core/syms_parser.c
@@ -172,6 +172,57 @@ syms_default_vbase_from_bin(SYMS_BinAccel *bin){
   return(result);
 }
 
+// imports & exports
+SYMS_API SYMS_ImportArray
+syms_imports_from_bin(SYMS_Arena *arena, SYMS_String8 data, SYMS_BinAccel *bin){
+  SYMS_ProfBegin("syms_imports_from_bin");
+  SYMS_ImportArray result = {0};
+  switch (bin->format){
+    case SYMS_FileFormat_PE:
+    {
+      result = syms_pe_imports_from_bin(arena, data, &bin->pe_accel);
+    }break;
+    
+    case SYMS_FileFormat_ELF:
+    {
+      result = syms_elf_imports_from_bin(arena, data, &bin->elf_accel);
+    }break;
+    
+    case SYMS_FileFormat_MACH:
+    {
+      result = syms_mach_imports_from_bin(arena, data, &bin->mach_accel);
+    }break;
+  }
+  SYMS_ASSERT_PARANOID(syms_parser_api_invariants());
+  SYMS_ProfEnd();
+  return(result);
+}
+
+SYMS_API SYMS_ExportArray
+syms_exports_from_bin(SYMS_Arena *arena, SYMS_String8 data, SYMS_BinAccel *bin){
+  SYMS_ProfBegin("syms_exports_from_bin");
+  SYMS_ExportArray result = {0};
+  switch (bin->format){
+    case SYMS_FileFormat_PE:
+    {
+      result = syms_pe_exports_from_bin(arena, data, &bin->pe_accel);
+    }break;
+    
+    case SYMS_FileFormat_ELF:
+    {
+      result = syms_elf_exports_from_bin(arena, data, &bin->elf_accel);
+    }break;
+    
+    case SYMS_FileFormat_MACH:
+    {
+      result = syms_mach_exports_from_bin(arena, data, &bin->mach_accel);
+    }break;
+  }
+  SYMS_ASSERT_PARANOID(syms_parser_api_invariants());
+  SYMS_ProfEnd();
+  return(result);
+}
+
 
 ////////////////////////////////
 //~ NOTE(nick): Bin List
@@ -1065,6 +1116,28 @@ syms_scope_children_from_sid(SYMS_Arena *arena, SYMS_String8 data, SYMS_DbgAccel
   return(result);
 }
 
+SYMS_API SYMS_StrippedInfoArray
+syms_stripped_from_unit(SYMS_Arena *arena, SYMS_String8 data, SYMS_DbgAccel *dbg, SYMS_UnitAccel *unit){
+  SYMS_ProfBegin("syms_stripped_from_unit");
+  SYMS_StrippedInfoArray result = {0};
+  if (unit->format == dbg->format){
+    switch (unit->format){
+      case SYMS_FileFormat_PDB:
+      {
+        result = syms_pdb_stripped_from_unit(arena, data, &dbg->pdb_accel, &unit->pdb_accel);
+      }break;
+      
+      case SYMS_FileFormat_DWARF:
+      {
+        // do nothing
+      }break;
+    }
+  }
+  SYMS_ASSERT_PARANOID(syms_parser_api_invariants());
+  SYMS_ProfEnd();
+  return(result);
+}
+
 // signature info
 SYMS_API SYMS_SigInfo
 syms_sig_info_from_type_sid(SYMS_Arena *arena, SYMS_String8 data, SYMS_DbgAccel *dbg,
diff --git a/Engine/Source/Runtime/SymsLib/syms/core/syms_parser.h b/Engine/Source/Runtime/SymsLib/syms/core/syms_parser.h
--- a/Engine/Source/Runtime/SymsLib/syms/core/syms_parser.h
+++ b/Engine/Source/Runtime/SymsLib/syms/core/syms_parser.h
@@ -104,14 +104,19 @@ SYMS_API SYMS_SecInfoArray syms_sec_info_array_from_bin(SYMS_Arena *arena, SYMS_
 // default vbase
 SYMS_API SYMS_U64          syms_default_vbase_from_bin(SYMS_BinAccel *bin);
 
+// imports & exports
+SYMS_API SYMS_ImportArray  syms_imports_from_bin(SYMS_Arena *arena, SYMS_String8 data, SYMS_BinAccel *bin);
+SYMS_API SYMS_ExportArray  syms_exports_from_bin(SYMS_Arena *arena, SYMS_String8 data, SYMS_BinAccel *bin);
+
 ////////////////////////////////
 //~ NOTE(nick): Bin List
 
-SYMS_API SYMS_B32           syms_file_is_bin_list(SYMS_FileAccel *file);
-SYMS_API SYMS_BinListAccel* syms_bin_list_from_file(SYMS_Arena *arena, SYMS_String8 data, SYMS_FileAccel *file);
-SYMS_API SYMS_BinInfoArray  syms_bin_info_array_from_bin_list(SYMS_Arena *arena, SYMS_BinListAccel *list);
-SYMS_API SYMS_BinAccel*     syms_bin_accel_from_bin_list_number(SYMS_Arena *arena, SYMS_String8 data,
-                                                                SYMS_BinListAccel *list, SYMS_U64 n);
+SYMS_API SYMS_B32          syms_file_is_bin_list(SYMS_FileAccel *file);
+SYMS_API SYMS_BinListAccel*syms_bin_list_from_file(SYMS_Arena *arena, SYMS_String8 data, SYMS_FileAccel *file);
+SYMS_API SYMS_BinInfoArray syms_bin_info_array_from_bin_list(SYMS_Arena *arena, SYMS_BinListAccel *list);
+SYMS_API SYMS_BinAccel*    syms_bin_accel_from_bin_list_number(SYMS_Arena *arena, SYMS_String8 data,
+                                                               SYMS_BinListAccel *list, SYMS_U64 n);
+
 ////////////////////////////////
 //~ NOTE(rjf): Dbg File
 
@@ -193,7 +198,8 @@ SYMS_API SYMS_SigInfo      syms_sig_info_from_mem_number(SYMS_Arena *arena, SYMS
                                                          SYMS_UnitAccel *unit, SYMS_MemsAccel *mems, SYMS_U64 n);
 
 SYMS_API SYMS_USID         syms_symbol_from_mem_number(SYMS_String8 data, SYMS_DbgAccel *dbg, SYMS_UnitAccel *unit, SYMS_MemsAccel *mems, SYMS_U64 n);
-SYMS_API SYMS_USID         syms_containing_type_from_sid(SYMS_String8 data, SYMS_DbgAccel *dbg, SYMS_UnitAccel *unit, SYMS_SymbolID sid);
+SYMS_API SYMS_USID         syms_containing_type_from_sid(SYMS_String8 data, SYMS_DbgAccel *dbg,
+                                                         SYMS_UnitAccel *unit, SYMS_SymbolID sid);
 // TODO(rjf): linkage name
 
 SYMS_API SYMS_EnumInfoArray syms_enum_info_array_from_sid(SYMS_Arena *arena, SYMS_String8 data, SYMS_DbgAccel *dbg,
@@ -211,6 +217,9 @@ SYMS_API SYMS_SymbolIDArray syms_scope_children_from_sid(SYMS_Arena *arena, SYMS
                                                          SYMS_DbgAccel *dbg, SYMS_UnitAccel *unit,
                                                          SYMS_SymbolID sid);
 
+SYMS_API SYMS_StrippedInfoArray syms_stripped_from_unit(SYMS_Arena *arena, SYMS_String8 data, SYMS_DbgAccel *dbg,
+                                                        SYMS_UnitAccel *unit);
+
 // signature info
 SYMS_API SYMS_SigInfo      syms_sig_info_from_type_sid(SYMS_Arena *arena, SYMS_String8 data, SYMS_DbgAccel *dbg,
                                                        SYMS_UnitAccel *unit, SYMS_SymbolID sid);
diff --git a/Engine/Source/Runtime/SymsLib/syms/core/unwind/syms_unwind_x64.h b/Engine/Source/Runtime/SymsLib/syms/core/unwind/syms_unwind_x64.h
new file mode 100644
--- /dev/null
+++ b/Engine/Source/Runtime/SymsLib/syms/core/unwind/syms_unwind_x64.h
@@ -0,0 +1,32 @@
+// Copyright Epic Games, Inc. All Rights Reserved.
+/* date = November 4th 2021 1:24 pm */
+
+#ifndef SYMS_UNWIND_X64_H
+#define SYMS_UNWIND_X64_H
+
+typedef struct SYMS_UnwindStateX64{
+  SYMS_U64 stack_pointer;
+  SYMS_U64 instruction_pointer;
+  
+  SYMS_U64 rax;
+  SYMS_U64 rcx;
+  SYMS_U64 rdx;
+  SYMS_U64 rbx;
+  SYMS_U64 rsp;
+  SYMS_U64 rbp;
+  SYMS_U64 rsi;
+  SYMS_U64 rdi;
+  SYMS_U64 r8;
+  SYMS_U64 r9;
+  SYMS_U64 r10;
+  SYMS_U64 r11;
+  SYMS_U64 r12;
+  SYMS_U64 r13;
+  SYMS_U64 r14;
+  SYMS_U64 r15;
+  
+  // TODO(allen): xmm
+  // TODO(allen): flags
+} SYMS_UnwindStateX64;
+
+#endif //SYMS_UNWIND_X64_H
diff --git a/Engine/Source/Runtime/SymsLib/syms/syms_inc.c b/Engine/Source/Runtime/SymsLib/syms/syms_inc.c
--- a/Engine/Source/Runtime/SymsLib/syms/syms_inc.c
+++ b/Engine/Source/Runtime/SymsLib/syms/syms_inc.c
@@ -29,6 +29,7 @@
 
 // "linux related" parsers
 #include "syms/core/elf/syms_elf_parser.c"
+#include "syms/core/dwarf/syms_dwarf_expr.c"
 #include "syms/core/dwarf/syms_dwarf_parser.c"
 
 // parser abstraction
@@ -38,6 +39,16 @@
 #include "syms/core/group/syms_group.c"
 #include "syms/core/file_inf/syms_file_inf.c"
 
+// regs
+#include "syms/core/regs/syms_regs.c"
+
+// depends on "syms_dwarf_expr" and "regs"
+#include "syms/core/regs/syms_dwarf_regs_helper.c"
+
+// unwinders
+#include "syms/core/unwind/syms_unwind_pe_x64.c"
+#include "syms/core/unwind/syms_unwind_elf_x64.c"
+
 // serialized type information
 #include "syms/syms_serial_inc.c"
 
diff --git a/Engine/Source/Runtime/SymsLib/syms/syms_inc.h b/Engine/Source/Runtime/SymsLib/syms/syms_inc.h
--- a/Engine/Source/Runtime/SymsLib/syms/syms_inc.h
+++ b/Engine/Source/Runtime/SymsLib/syms/syms_inc.h
@@ -31,6 +31,7 @@
 
 // "linux related" parsers
 #include "syms/core/elf/syms_elf_parser.h"
+#include "syms/core/dwarf/syms_dwarf_expr.h"
 #include "syms/core/dwarf/syms_dwarf_parser.h"
 
 // parser abstraction
@@ -49,6 +50,14 @@
 #include "syms/core/generated/syms_meta_regs_x64.h"
 #pragma pack(pop)
 
+#include "syms/core/regs/syms_regs_x64.h"
+// depends on "syms_dwarf_expr" and "regs"
+#include "syms/core/regs/syms_dwarf_regs_helper.h"
+
+// unwinders
+#include "syms/core/unwind/syms_unwind_pe_x64.h"
+#include "syms/core/unwind/syms_unwind_elf_x64.h"
+
 // serialized type information
 #include "syms/syms_serial_inc.h"
 
