Update SymsLib to latest version and minor code cleanup in SymslibResolver.
Includes fix for correctly zero-terminating stripped pdb symbol names.
Removes unneeded comparisons & useless calls when parsing or resolving symbols.

diff --git a/Engine/Source/Developer/TraceServices/Private/Model/SymslibResolver.cpp b/Engine/Source/Developer/TraceServices/Private/Model/SymslibResolver.cpp
--- a/Engine/Source/Developer/TraceServices/Private/Model/SymslibResolver.cpp
+++ b/Engine/Source/Developer/TraceServices/Private/Model/SymslibResolver.cpp
@@ -736,19 +736,11 @@ FSymslibResolver::EModuleStatus FSymslibResolver::LoadModule(FModuleEntry* Modul
 			SYMS_SymbolIDArray* ProcArray = syms_group_proc_sid_array_from_uid(Group, UnitID);
 			SYMS_U64 ProcCount = ProcArray->count;
 
-			uint32 UnitSymbolCount = 0;
-
 			FSymsSymbol* Symbols = syms_push_array(Arena, FSymsSymbol, ProcCount);
 			for (SYMS_U64 ProcIndex = 0; ProcIndex < ProcCount; ProcIndex++)
 			{
 				SYMS_SymbolID SymbolID = ProcArray->ids[ProcIndex];
 
-				SYMS_SymbolKind Kind = syms_group_symbol_kind_from_sid(Group, UnitAccel, SymbolID);
-				if (Kind == SYMS_SymbolKind_Procedure)
-				{
-					UnitSymbolCount++;
-				}
-
 				SYMS_String8 Name = syms_group_symbol_name_from_sid(Arena, Group, UnitAccel, SymbolID);
 				Symbols[ProcIndex].Name = reinterpret_cast<char*>(Name.str);
 
@@ -764,7 +756,7 @@ FSymslibResolver::EModuleStatus FSymslibResolver::LoadModule(FModuleEntry* Modul
 
 			syms_release_scratch(Scratch);
 
-			SymbolCount += UnitSymbolCount;
+			SymbolCount += ProcCount;
 		});
 		syms_group_end_multilane(Group);
 
@@ -776,24 +768,19 @@ FSymslibResolver::EModuleStatus FSymslibResolver::LoadModule(FModuleEntry* Modul
 	// store stripped format symbols
 	{
 		SYMS_StrippedInfoArray StrippedInfo = syms_group_stripped_info(Group);
-		SYMS_SpatialMap1D* StrippedMap = syms_group_stripped_info_map(Group);
 
 		FSymsSymbol* StrippedSymbols = syms_push_array(Arena, FSymsSymbol, StrippedInfo.count);
+		for (SYMS_U64 Index = 0; Index < StrippedInfo.count; Index++)
 		{
-			SYMS_StrippedInfo* InfoPtr = StrippedInfo.info;
-			FSymsSymbol* StrippedSymbolPtr = StrippedSymbols;
-			for (SYMS_U64 Index = 0; Index < StrippedInfo.count; Index++)
-			{
-				SYMS_String8 Name = syms_push_string_copy(Arena, InfoPtr->name);
-				StrippedSymbolPtr->Name = reinterpret_cast<char*>(Name.str);
-				++StrippedSymbolPtr;
-				++InfoPtr;
-			}
+			SYMS_StrippedInfo* Info = &StrippedInfo.info[Index];
+			FSymsSymbol* StrippedSymbol = &StrippedSymbols[Index];
+
+			SYMS_String8 Name = syms_push_string_copy(Arena, Info->name);
+			StrippedSymbol->Name = reinterpret_cast<char*>(Name.str);
 		}
 
-		Instance->StrippedMap = syms_spatial_map_1d_copy(Arena, StrippedMap);
+		Instance->StrippedMap = syms_spatial_map_1d_copy(Arena, syms_group_stripped_info_map(Group));
 		Instance->StrippedSymbols = StrippedSymbols;
-		Instance->StrippedCount = StrippedInfo.count;
 
 		SymbolCount += StrippedInfo.count;
 	}
@@ -862,7 +849,7 @@ bool FSymslibResolver::ResolveSymbol(uint64 Address, FResolvedSymbol* Target, FS
 	uint32 SourceFileLine = 0;
 
 	SYMS_UnitID UnitID = syms_spatial_map_1d_value_from_point(&Instance->UnitMap, VirtualOffset);
-	if (UnitID > 0 && UnitID <= Instance->Units.Num())
+	if (UnitID)
 	{
 		FSymsUnit* Unit = &Instance->Units[UnitID - 1];
 
@@ -875,7 +862,7 @@ bool FSymslibResolver::ResolveSymbol(uint64 Address, FResolvedSymbol* Target, FS
 			if (SeqNumber)
 			{
 				SYMS_Line Line = syms_line_from_sequence_voff(&Unit->LineTable, SeqNumber, VirtualOffset);
-				if (Line.src_coord.file_id > 0 && Line.src_coord.file_id <= Unit->FileTable.count)
+				if (Line.src_coord.file_id)
 				{
 					SYMS_String8 FileName = Unit->FileTable.strings[Line.src_coord.file_id - 1];
 
@@ -894,10 +881,9 @@ bool FSymslibResolver::ResolveSymbol(uint64 Address, FResolvedSymbol* Target, FS
 	{
 		// try lookup into stripped format symbols
 		SYMS_U64 Value = syms_spatial_map_1d_value_from_point(&Instance->StrippedMap, VirtualOffset);
-		if (Value > 0 && Value <= Instance->StrippedCount)
+		if (Value)
 		{
-			SYMS_U64 Index = Value - 1;
-			SymsSymbol = &Instance->StrippedSymbols[Index];
+			SymsSymbol = &Instance->StrippedSymbols[Value - 1];
 
 			// use module name as filename
 			SourceFilePersistent = StringAllocator.Store(Module->Name);
diff --git a/Engine/Source/Developer/TraceServices/Private/Model/SymslibResolver.h b/Engine/Source/Developer/TraceServices/Private/Model/SymslibResolver.h
--- a/Engine/Source/Developer/TraceServices/Private/Model/SymslibResolver.h
+++ b/Engine/Source/Developer/TraceServices/Private/Model/SymslibResolver.h
@@ -56,7 +56,6 @@ private:
 		SYMS_SpatialMap1D   UnitMap;
 		SYMS_SpatialMap1D   StrippedMap;
 		FSymsSymbol*        StrippedSymbols;
-		uint64              StrippedCount;
 		uint64              DefaultBase;
 	};
 
diff --git a/Engine/Source/Runtime/SymsLib/syms/core/base/syms_base.c b/Engine/Source/Runtime/SymsLib/syms/core/base/syms_base.c
--- a/Engine/Source/Runtime/SymsLib/syms/core/base/syms_base.c
+++ b/Engine/Source/Runtime/SymsLib/syms/core/base/syms_base.c
@@ -64,8 +64,8 @@ syms_hash_u64(SYMS_U64 x){
 ////////////////////////////////
 //~ rjf: Serial Information Functions
 
-SYMS_GLOBAL SYMS_SerialValue syms_null_serial_value;
-SYMS_GLOBAL SYMS_SerialFlag syms_null_serial_flag;
+SYMS_READ_ONLY SYMS_GLOBAL SYMS_SerialValue syms_null_serial_value;
+SYMS_READ_ONLY SYMS_GLOBAL SYMS_SerialFlag syms_null_serial_flag;
 
 SYMS_API SYMS_SerialField*
 syms_serial_first_field(SYMS_SerialType *type){
@@ -326,7 +326,7 @@ syms_string_list_join(SYMS_Arena *arena, SYMS_String8List *list, SYMS_StringJoin
   if (list->node_count >= 2){
     result.size += (list->node_count - 1)*join.sep.size;
   }
-  result.str = syms_push_array(arena, SYMS_U8, result.size);
+  result.str = syms_push_array(arena, SYMS_U8, result.size + 1);
   
   // fill
   SYMS_U8 *ptr = result.str;
@@ -345,6 +345,7 @@ syms_string_list_join(SYMS_Arena *arena, SYMS_String8List *list, SYMS_StringJoin
   }
   syms_memmove(ptr, join.post.str, join.post.size);
   ptr += join.post.size;
+  *ptr = 0;
   
   return(result);
 }
diff --git a/Engine/Source/Runtime/SymsLib/syms/core/base/syms_base.h b/Engine/Source/Runtime/SymsLib/syms/core/base/syms_base.h
--- a/Engine/Source/Runtime/SymsLib/syms/core/base/syms_base.h
+++ b/Engine/Source/Runtime/SymsLib/syms/core/base/syms_base.h
@@ -41,6 +41,22 @@
 #define SYMS_GLOBAL static
 #define SYMS_LOCAL  static
 
+#ifndef SYMS_READ_ONLY
+# if SYMS_COMPILER_CL || (SYMS_COMPILER_CLANG && SYMS_OS_WINDOWS)
+#  pragma section(".roglob", read)
+#  define SYMS_READ_ONLY __declspec(allocate(".roglob"))
+# elif (SYMS_COMPILER_CLANG && SYMS_OS_LINUX)
+#  define SYMS_READ_ONLY __attribute__((section(".rodata")))
+# else
+// NOTE(rjf): I don't know of a useful way to do this in GCC land.
+// __attribute__((section(".rodata"))) looked promising, but it introduces a
+// strange warning about malformed section attributes, and it doesn't look
+// like writing to that section reliably produces access violations, strangely
+// enough. (It does on Clang)
+#  define SYMS_READ_ONLY
+# endif
+#endif
+
 #if SYMS_COMPILER_CL
 # define SYMS_THREAD_LOCAL __declspec(thread)
 #elif SYMS_COMPILER_CLANG || SYMS_COMPILER_GCC
@@ -400,6 +416,7 @@ SYMS_API SYMS_U64 syms_hash_u64(SYMS_U64 x);
 #define syms_serial_type(name) (_syms_serial_type_##name)
 #define syms_string_from_enum_value(enum_type, value) \
 (syms_serial_value_from_enum_value(&syms_serial_type(enum_type), value)->name)
+#define syms_bswap_in_place(type, ptr) syms_bswap_in_place__##type(ptr)
 
 SYMS_API SYMS_SerialField* syms_serial_first_field(SYMS_SerialType *type);
 SYMS_API SYMS_SerialValue* syms_serial_first_value(SYMS_SerialType *type);
diff --git a/Engine/Source/Runtime/SymsLib/syms/core/group/syms_group.c b/Engine/Source/Runtime/SymsLib/syms/core/group/syms_group.c
--- a/Engine/Source/Runtime/SymsLib/syms/core/group/syms_group.c
+++ b/Engine/Source/Runtime/SymsLib/syms/core/group/syms_group.c
@@ -1503,15 +1503,6 @@ syms_line_to_addr_line_sort__rec(SYMS_FileToLineToAddrLooseLine **array, SYMS_U6
 ////////////////////////////////
 // NOTE(allen): Syms Group Type Graph Functions
 
-SYMS_GLOBAL SYMS_TypeNode syms_type_graph_nil = {
-  SYMS_TypeKind_Null,     // kind
-  {(SYMS_U8*)"(nil)", 5}, // name
-  0,                      // byte_size
-  0,                      // src_coord
-  &syms_type_graph_nil,   // direct_type
-  &syms_type_graph_nil,   // this_type
-};
-
 SYMS_API SYMS_TypeNode*
 syms_group_type_from_usid(SYMS_Group *group, SYMS_USID usid){
   SYMS_ProfBegin("syms_group_type_from_usid");
diff --git a/Engine/Source/Runtime/SymsLib/syms/core/group/syms_group.h b/Engine/Source/Runtime/SymsLib/syms/core/group/syms_group.h
--- a/Engine/Source/Runtime/SymsLib/syms/core/group/syms_group.h
+++ b/Engine/Source/Runtime/SymsLib/syms/core/group/syms_group.h
@@ -222,13 +222,22 @@ typedef struct SYMS_Group{
 ////////////////////////////////
 // NOTE(allen): Data Structure Nils
 
-SYMS_GLOBAL SYMS_SymbolIDArray syms_sid_array_nil = {0};
-SYMS_GLOBAL SYMS_String8Array syms_string_array_nil = {0};
-SYMS_GLOBAL SYMS_LineParseOut syms_line_parse_nil = {0};
-SYMS_GLOBAL SYMS_SpatialMap1D syms_spatial_map_1d_nil = {0};
-SYMS_GLOBAL SYMS_LineToAddrMap syms_line_to_addr_map_nil = {0};
-SYMS_GLOBAL SYMS_TypeMemberArray syms_type_member_array_nil = {0};
-SYMS_GLOBAL SYMS_EnumInfoArray syms_enum_info_array_nil = {0};
+SYMS_READ_ONLY SYMS_GLOBAL SYMS_SymbolIDArray syms_sid_array_nil = {0};
+SYMS_READ_ONLY SYMS_GLOBAL SYMS_String8Array syms_string_array_nil = {0};
+SYMS_READ_ONLY SYMS_GLOBAL SYMS_LineParseOut syms_line_parse_nil = {0};
+SYMS_READ_ONLY SYMS_GLOBAL SYMS_SpatialMap1D syms_spatial_map_1d_nil = {0};
+SYMS_READ_ONLY SYMS_GLOBAL SYMS_LineToAddrMap syms_line_to_addr_map_nil = {0};
+SYMS_READ_ONLY SYMS_GLOBAL SYMS_TypeMemberArray syms_type_member_array_nil = {0};
+SYMS_READ_ONLY SYMS_GLOBAL SYMS_EnumInfoArray syms_enum_info_array_nil = {0};
+
+SYMS_READ_ONLY SYMS_GLOBAL SYMS_TypeNode syms_type_graph_nil = {
+  SYMS_TypeKind_Null,     // kind
+  {(SYMS_U8*)"(nil)", 5}, // name
+  0,                      // byte_size
+  0,                      // src_coord
+  &syms_type_graph_nil,   // direct_type
+  &syms_type_graph_nil,   // this_type
+};
 
 
 ////////////////////////////////
diff --git a/Engine/Source/Runtime/SymsLib/syms/core/mach/syms_mach_parser.c b/Engine/Source/Runtime/SymsLib/syms/core/mach/syms_mach_parser.c
--- a/Engine/Source/Runtime/SymsLib/syms/core/mach/syms_mach_parser.c
+++ b/Engine/Source/Runtime/SymsLib/syms/core/mach/syms_mach_parser.c
@@ -48,14 +48,14 @@ syms_mach_bin_from_base_range(SYMS_Arena *arena, void *base, SYMS_U64Range range
       SYMS_MachHeader32 header32 = {0};
       after_header_off = syms_based_range_read_struct(base, range, 0, &header32);
       if (is_swapped){
-        syms_mach_header32_endian_swap_in_place(&header32);
+        syms_bswap_in_place(SYMS_MachHeader32, &header32);
       }
       syms_mach_header64_from_header32(&header, &header32);
     }
     else{
       after_header_off = syms_based_range_read_struct(base, range, 0, &header);
       if (is_swapped){
-        syms_mach_header64_endian_swap_in_place(&header);
+        syms_bswap_in_place(SYMS_MachHeader64, &header);
       }
     }
     
@@ -70,13 +70,16 @@ syms_mach_bin_from_base_range(SYMS_Arena *arena, void *base, SYMS_U64Range range
     SYMS_MachSectionNode *section_last = 0;
     SYMS_U32 section_count = 0;
     
-    SYMS_U64Range regular_bind_range = syms_make_u64_range(0,0);
-    SYMS_U64Range lazy_bind_range    = syms_make_u64_range(0,0);
-    SYMS_U64Range weak_bind_range    = syms_make_u64_range(0,0);
-    SYMS_U64Range export_range       = syms_make_u64_range(0,0);
+    SYMS_U64Range bind_ranges[SYMS_MachBindTable_COUNT]; syms_memzero_struct(&bind_ranges[0]);
+    SYMS_U64Range export_range = syms_make_u64_range(0,0);
     
     SYMS_MachDylibList dylib_list; syms_memzero_struct(&dylib_list);
     
+    { // push null dylib
+      SYMS_MachDylib dylib; syms_memzero_struct(&dylib);
+      syms_mach_dylib_list_push(scratch.arena, &dylib_list, &dylib, syms_make_u64_range(0,0));
+    }
+    
     SYMS_U64 next_cmd_off = after_header_off;
     for (SYMS_U32 i = 0; i < header.ncmds; i += 1){
       //- align read offset
@@ -100,7 +103,7 @@ syms_mach_bin_from_base_range(SYMS_Arena *arena, void *base, SYMS_U64Range range
           SYMS_MachSegmentCommand64 segment_command64 = {0};
           syms_based_range_read_struct(base, range, cmd_off, &segment_command64);
           if (is_swapped){
-            syms_mach_segment_command64_endian_swap_in_place(&segment_command64);
+            syms_bswap_in_place(SYMS_MachSegmentCommand64, &segment_command64);
           }
           SYMS_U64 after_seg_off = cmd_off + sizeof(SYMS_MachSegmentCommand64);
           
@@ -119,7 +122,7 @@ syms_mach_bin_from_base_range(SYMS_Arena *arena, void *base, SYMS_U64Range range
             SYMS_MachSection64 section64 = {0};
             syms_based_range_read_struct(base, range, sec_off, &section64);
             if (is_swapped){
-              syms_mach_section64_endian_swap_in_place(&section64);
+              syms_bswap_in_place(SYMS_MachSection64, &section64);
             }
             next_sec_off = sec_off + sizeof(SYMS_MachSection64);
             
@@ -136,28 +139,29 @@ syms_mach_bin_from_base_range(SYMS_Arena *arena, void *base, SYMS_U64Range range
           SYMS_MachDyldInfoCommand dyld;
           syms_based_range_read_struct(base, range, cmd_off, &dyld);
           if (is_swapped){
-            syms_mach_dyld_info_command_endian_swap_in_place(&dyld);
+            syms_bswap_in_place(SYMS_MachDyldInfoCommand, &dyld);
           }
           
-          regular_bind_range = syms_make_u64_inrange(range, dyld.bind_off, dyld.bind_size);
-          lazy_bind_range    = syms_make_u64_inrange(range, dyld.lazy_bind_off, dyld.lazy_bind_size);
-          weak_bind_range    = syms_make_u64_inrange(range, dyld.weak_bind_off, dyld.weak_bind_size);
-          export_range       = syms_make_u64_inrange(range, dyld.export_off, dyld.export_size);
+          bind_ranges[SYMS_MachBindTable_REGULAR] = syms_make_u64_inrange(range, dyld.bind_off, dyld.bind_size);
+          bind_ranges[SYMS_MachBindTable_LAZY]    = syms_make_u64_inrange(range, dyld.lazy_bind_off, dyld.lazy_bind_size);
+          bind_ranges[SYMS_MachBindTable_WEAK]    = syms_make_u64_inrange(range, dyld.weak_bind_off, dyld.weak_bind_size);
+          export_range = syms_make_u64_inrange(range, dyld.export_off, dyld.export_size);
         }break;
         
-        case SYMS_MachLoadCommandType_ID_DYLIB:
+        case SYMS_MachLoadCommandType_LOAD_WEAK_DYLIB:
         case SYMS_MachLoadCommandType_LOAD_DYLIB: 
+        case SYMS_MachLoadCommandType_LOAD_UPWARD_DYLIB:
+        case SYMS_MachLoadCommandType_LAZY_LOAD_DYLIB:
         {
           SYMS_MachDylibCommand cmd; syms_memzero_struct(&cmd);
           syms_based_range_read_struct(base, range, cmd_off, &cmd);
           SYMS_MachDylib *dylib = &cmd.dylib;
           if (is_swapped){
-            syms_mach_dylib_endian_swap_in_place(dylib);
+            syms_bswap_in_place(SYMS_MachDylib, dylib);
           }
-          
-          SYMS_U64 name_offset = cmd_off + dylib->name.offset;
-          SYMS_U64Range name = syms_make_u64_inrange(range, name_offset, (cmd_off + lc.size) - name_offset);
-          syms_mach_dylib_list_push(scratch.arena, &dylib_list, dylib, name);
+          SYMS_U64      name_offset = cmd_off + dylib->name.offset;
+          SYMS_U64Range name_range  = syms_make_u64_inrange(range, name_offset, (cmd_off + lc.size) - name_offset);
+          syms_mach_dylib_list_push(scratch.arena, &dylib_list, dylib, name_range);
         }break;
       }
       
@@ -187,17 +191,15 @@ syms_mach_bin_from_base_range(SYMS_Arena *arena, void *base, SYMS_U64Range range
     }
     
     //- fill result
-    result                     = syms_push_array(arena, SYMS_MachBinAccel, 1);
-    result->format             = SYMS_FileFormat_MACH;
-    result->arch               = syms_mach_arch_from_cputype(header.cputype);
-    result->segment_count      = segment_count;
-    result->segments           = segments;
-    result->section_count      = section_count;
-    result->sections           = sections;
-    result->regular_bind_range = regular_bind_range;
-    result->lazy_bind_range    = lazy_bind_range;
-    result->weak_bind_range    = weak_bind_range;
-    result->export_range       = export_range;
+    result                = syms_push_array(arena, SYMS_MachBinAccel, 1);
+    result->format        = SYMS_FileFormat_MACH;
+    result->arch          = syms_mach_arch_from_cputype(header.cputype);
+    result->segment_count = segment_count;
+    result->segments      = segments;
+    result->section_count = section_count;
+    result->sections      = sections;
+    syms_memmove(&result->bind_ranges[0], &bind_ranges[0], sizeof(bind_ranges));
+    result->export_range  = export_range;
     { // convert dylib list to array
       result->dylib_count = 0;
       result->dylibs = syms_push_array(arena, SYMS_MachParsedDylib, dylib_list.count);
@@ -297,7 +299,7 @@ syms_mach_bin_list_accel_from_file(SYMS_Arena *arena, SYMS_String8 data, SYMS_Ma
     SYMS_MachFatHeader fat_header = {0};
     read_offset += syms_based_range_read_struct(base, range, read_offset, &fat_header);
     if (is_swapped){
-      syms_mach_fat_header_endian_swap_in_place(&fat_header);
+      syms_bswap_in_place(SYMS_MachFatHeader, &fat_header);
     }
     
     SYMS_U64 fat_count = fat_header.nfat_arch;
@@ -307,7 +309,7 @@ syms_mach_bin_list_accel_from_file(SYMS_Arena *arena, SYMS_String8 data, SYMS_Ma
     for (SYMS_U32 i = 0; i < fat_count; i += 1, fat_ptr += 1){
       read_offset += syms_based_range_read_struct(base, range, read_offset, fat_ptr);
       if (is_swapped){
-        syms_mach_fat_arch_endian_swap_in_place(fat_ptr);
+        syms_bswap_in_place(SYMS_MachFatArch, fat_ptr);
       }
     }
     
@@ -476,15 +478,21 @@ syms_mach_binds_from_base_range(SYMS_Arena *arena, void *base, SYMS_U64Range ran
       } break;
       case SYMS_MachBindOpcode_DO_BIND_ADD_ADDR_ULEB:
       {
-        SYMS_ASSERT_PARANOID(!"TODO: DO_BIND_ADD_ADDR_ULEB");
+        SYMS_U64 addend = 0;
+        read_offset += syms_based_range_read_uleb128(base, range, read_offset, &addend);
+        state.segment_offset += addend;
       } goto do_bind;
       case SYMS_MachBindOpcode_DO_BIND_ADD_ADDR_IMM_SCALED:
       {
-        SYMS_ASSERT_PARANOID(!"TODO: DO_BIND_ADD_ADDR_IMM_SCALED");
+        state.segment_offset += imm;
       } goto do_bind;
       case SYMS_MachBindOpcode_DO_BIND_ULEB_TIMES_SKIPPING_ULEB:
       {
-        SYMS_ASSERT_PARANOID(!"TODO: DO_BIND_ULEB_TIMES_SKIPPING_ULEB");
+        SYMS_U64 count = 0, skip = 0;
+        read_offset += syms_based_range_read_uleb128(base, range, read_offset, &count);
+        read_offset += syms_based_range_read_uleb128(base, range, read_offset, &skip);
+        
+        state.segment_offset += skip * address_size;
       } goto do_bind;
       case SYMS_MachBindOpcode_DO_BIND:
       {
@@ -517,7 +525,7 @@ syms_mach_binds_from_base_range(SYMS_Arena *arena, void *base, SYMS_U64Range ran
 SYMS_API SYMS_ImportArray
 syms_mach_imports_from_bin(SYMS_Arena *arena, SYMS_String8 data, SYMS_MachBinAccel *bin){
   SYMS_ArenaTemp scratch = syms_get_scratch(&arena, 1);
-
+  
   // read and copy library names
   SYMS_String8 *dylib_names = syms_push_array(scratch.arena, SYMS_String8, bin->dylib_count);
   for (SYMS_U32 i = 0; i < bin->dylib_count; ++i){
@@ -525,17 +533,16 @@ syms_mach_imports_from_bin(SYMS_Arena *arena, SYMS_String8 data, SYMS_MachBinAcc
     dylib_names[i] = syms_push_string_copy(arena, name);
   }
   
+  // build binds
   SYMS_U64 address_size = syms_address_size_from_arch(bin->arch);
-  SYMS_MachBindList binds[3];
-  binds[SYMS_MachBindTable_REGULAR] = syms_mach_binds_from_base_range(scratch.arena, (void*)data.str, bin->regular_bind_range, address_size, SYMS_MachBindTable_REGULAR);
-  binds[SYMS_MachBindTable_LAZY]    = syms_mach_binds_from_base_range(scratch.arena, (void*)data.str, bin->lazy_bind_range, address_size, SYMS_MachBindTable_LAZY);
-  binds[SYMS_MachBindTable_WEAK]    = syms_mach_binds_from_base_range(scratch.arena, (void*)data.str, bin->weak_bind_range, address_size, SYMS_MachBindTable_WEAK);
-  
   SYMS_U64 total_bind_count = 0;
+  SYMS_MachBindList binds[SYMS_MachBindTable_COUNT];
   for (SYMS_U32 i = 0; i < SYMS_ARRAY_SIZE(binds); ++i){
+    binds[i] = syms_mach_binds_from_base_range(scratch.arena, (void*)data.str, bin->bind_ranges[i], address_size, (SYMS_MachBindTable)i);
     total_bind_count += binds[i].count;
   }
-
+  
+  // convert mach binds to syms imports
   SYMS_ImportArray import_array; 
   import_array.count = 0;
   import_array.imports = syms_push_array(arena, SYMS_Import, total_bind_count);
@@ -560,18 +567,22 @@ syms_mach_imports_from_bin(SYMS_Arena *arena, SYMS_String8 data, SYMS_MachBinAcc
 static SYMS_MachExport *
 syms_mach_parse_export_node(SYMS_Arena *arena, void *base, SYMS_U64Range range, SYMS_String8 name, SYMS_U64 read_offset){
   SYMS_MachExport *node = 0;
-  
   SYMS_U64 export_size = 0;
   SYMS_U64 read_size = syms_based_range_read_uleb128(base, range, read_offset, &export_size);
-  read_offset += read_size;
-  
   if (read_size > 0){
+    read_offset += read_size;
+    
+    // push new node
     node = syms_push_array_zero(arena, SYMS_MachExport, 1);
     node->name = name;
+    node->is_export_info = export_size != 0;
     
-    SYMS_B32 is_export_info = export_size != 0;
-    if (is_export_info){
-      SYMS_U64Range export_range = syms_make_u64_inrange(range, read_offset, export_size);
+    // read export discriminator
+    SYMS_U64Range export_range = syms_make_u64_inrange(range, read_offset, export_size);
+    read_offset += export_size;
+    
+    // parse export info
+    if (node->is_export_info){
       SYMS_U64 export_read_offset = 0;
       export_read_offset += syms_based_range_read_uleb128(base, export_range, export_read_offset, &node->flags);
       if (node->flags & SYMS_MachExportSymbolFlags_REEXPORT){
@@ -586,6 +597,7 @@ syms_mach_parse_export_node(SYMS_Arena *arena, void *base, SYMS_U64Range range,
       }
     }
     
+    // read child info
     read_offset += syms_based_range_read_struct(base, range, read_offset, &node->child_count);
     if (node->child_count > 0){
       node->children = syms_push_array(arena, SYMS_MachExport *, node->child_count);
@@ -611,7 +623,7 @@ syms_build_mach_export_trie(SYMS_Arena *arena, void *base, SYMS_U64Range range){
 SYMS_API SYMS_ExportArray
 syms_mach_exports_from_bin(SYMS_Arena *arena, SYMS_String8 data, SYMS_MachBinAccel *bin){
   SYMS_ArenaTemp scratch = syms_get_scratch(&arena, 1);
-
+  
   // read and copy library names
   SYMS_String8 *dylib_names = syms_push_array(scratch.arena, SYMS_String8, bin->dylib_count);
   for (SYMS_U32 i = 0; i < bin->dylib_count; ++i){
@@ -619,7 +631,10 @@ syms_mach_exports_from_bin(SYMS_Arena *arena, SYMS_String8 data, SYMS_MachBinAcc
     dylib_names[i] = syms_push_string_copy(arena, name);
   }
   
+  // build export trie
   SYMS_MachExport *root_export = syms_build_mach_export_trie(scratch.arena, (void*)data.str, bin->export_range);
+  
+  // traverse trie
   SYMS_ExportNode *first_export = 0;
   SYMS_ExportNode *last_export = 0;
   SYMS_U32 export_count = 0;
@@ -631,9 +646,11 @@ syms_mach_exports_from_bin(SYMS_Arena *arena, SYMS_String8 data, SYMS_MachBinAcc
     frame->node = root_export;
     SYMS_StackPush(stack, frame);
     
-    for (;;){
-      if (frame->child_idx >= frame->node->child_count){
-        if (frame->node->child_count == 0){ // leaf
+    for (; stack != 0; ){
+      SYMS_B32 pop_frame = frame->child_idx >= frame->node->child_count;
+      if (pop_frame){
+        if (frame->node->is_export_info){
+          // append strings from nodes on the stack
           SYMS_ArenaTemp temp = syms_arena_temp_begin(scratch.arena);
           SYMS_String8List name_parts; syms_memzero_struct(&name_parts);
           for (SYMS_MachExportFrame *f = frame; f != 0; f = f->next){
@@ -644,8 +661,12 @@ syms_mach_exports_from_bin(SYMS_Arena *arena, SYMS_String8 data, SYMS_MachBinAcc
           syms_arena_temp_end(temp);
           
           SYMS_ExportNode *export_node = syms_push_array(scratch.arena, SYMS_ExportNode, 1); 
+          
+          // add export to the list
           SYMS_QueuePush(first_export, last_export, export_node);
           export_count += 1;
+          
+          // initialize export
           SYMS_Export *exp = &export_node->data;
           if (frame->node->flags & SYMS_MachExportSymbolFlags_REEXPORT){
             // TODO(nick): need dylib with reexports to test this code path
@@ -666,18 +687,16 @@ syms_mach_exports_from_bin(SYMS_Arena *arena, SYMS_String8 data, SYMS_MachBinAcc
             exp->forwarder_import_name = syms_str8(0,0);
           }
         }
+        
         // pop current frame 
         SYMS_MachExportFrame *to_free = frame;
         frame = SYMS_StackPop(stack);
+        
         // push frame to free list for later reuse
         to_free->next = 0;
         SYMS_StackPush(free_frames, to_free);
-        if (frame == 0){
-          break; // we traversed entire trie, leave
-        }
-      }
-      if (frame->child_idx < frame->node->child_count){
-        // push frame with child node
+      } else{
+        // push frame
         SYMS_MachExportFrame *new_frame = free_frames;
         if (new_frame){
           SYMS_StackPop(free_frames);
diff --git a/Engine/Source/Runtime/SymsLib/syms/core/mach/syms_mach_parser.h b/Engine/Source/Runtime/SymsLib/syms/core/mach/syms_mach_parser.h
--- a/Engine/Source/Runtime/SymsLib/syms/core/mach/syms_mach_parser.h
+++ b/Engine/Source/Runtime/SymsLib/syms/core/mach/syms_mach_parser.h
@@ -27,6 +27,15 @@ typedef struct SYMS_MachSectionNode{
   SYMS_MachSection64 data;
 } SYMS_MachSectionNode;
 
+enum
+{
+  SYMS_MachBindTable_REGULAR,
+  SYMS_MachBindTable_LAZY,
+  SYMS_MachBindTable_WEAK,
+  SYMS_MachBindTable_COUNT
+};
+typedef SYMS_U32 SYMS_MachBindTable;
+
 typedef struct SYMS_MachBinAccel{
   SYMS_FileFormat format;
   
@@ -38,9 +47,7 @@ typedef struct SYMS_MachBinAccel{
   SYMS_MachSection64 *sections;
   SYMS_U32 section_count;
 
-  SYMS_U64Range regular_bind_range;
-  SYMS_U64Range lazy_bind_range;
-  SYMS_U64Range weak_bind_range;
+  SYMS_U64Range bind_ranges[SYMS_MachBindTable_COUNT];
   SYMS_U64Range export_range;
 
   SYMS_U32 dylib_count;
@@ -73,14 +80,6 @@ typedef struct
 ////////////////////////////////
 // Binds
 
-enum
-{
-  SYMS_MachBindTable_REGULAR,
-  SYMS_MachBindTable_LAZY,
-  SYMS_MachBindTable_WEAK
-};
-typedef SYMS_U32 SYMS_MachBindTable;
-
 typedef struct
 {
   SYMS_U32 segment;
@@ -113,6 +112,8 @@ typedef struct SYMS_MachExport
   struct SYMS_MachExport **children;
   SYMS_U8 child_count;
   
+  SYMS_B32 is_export_info;
+
   SYMS_String8 name;
   SYMS_U64 flags;
   SYMS_U64 address;
diff --git a/Engine/Source/Runtime/SymsLib/syms/core/pdb/syms_pdb_parser.c b/Engine/Source/Runtime/SymsLib/syms/core/pdb/syms_pdb_parser.c
--- a/Engine/Source/Runtime/SymsLib/syms/core/pdb/syms_pdb_parser.c
+++ b/Engine/Source/Runtime/SymsLib/syms/core/pdb/syms_pdb_parser.c
@@ -4239,7 +4239,7 @@ syms_pdb_stripped_from_unit(SYMS_Arena *arena, SYMS_String8 data, SYMS_PdbDbgAcc
       }
       
       // fill stripped info
-      info_ptr->name = name;
+      info_ptr->name = syms_push_string_copy(arena, name);
       info_ptr->voff = voff;
     }
     
diff --git a/Engine/Source/Runtime/SymsLib/syms/core/syms_debug_info.c b/Engine/Source/Runtime/SymsLib/syms/core/syms_debug_info.c
--- a/Engine/Source/Runtime/SymsLib/syms/core/syms_debug_info.c
+++ b/Engine/Source/Runtime/SymsLib/syms/core/syms_debug_info.c
@@ -61,6 +61,24 @@ syms_type_kind_is_basic(SYMS_TypeKind kind){
   return(result);
 }
 
+SYMS_API SYMS_B32
+syms_type_kind_is_integer(SYMS_TypeKind kind){
+  SYMS_B32 result = (SYMS_TypeKind_Int8 <= kind && kind <= SYMS_TypeKind_UInt128);
+  return(result);
+}
+
+SYMS_API SYMS_B32
+syms_type_kind_is_signed(SYMS_TypeKind kind){
+  SYMS_B32 result = (SYMS_TypeKind_Int8 <= kind && kind <= SYMS_TypeKind_Int128);
+  return(result);
+}
+
+SYMS_API SYMS_B32
+syms_type_kind_is_complex(SYMS_TypeKind kind){
+  SYMS_B32 result = (SYMS_TypeKind_Complex32 <= kind && kind <= SYMS_TypeKind_Complex128);
+  return(result);
+}
+
 SYMS_API SYMS_B32
 syms_type_kind_is_user_defined(SYMS_TypeKind kind){
   SYMS_B32 result = (SYMS_TypeKind_Struct <= kind && kind <= SYMS_TypeKind_ForwardEnum);
diff --git a/Engine/Source/Runtime/SymsLib/syms/core/syms_debug_info.h b/Engine/Source/Runtime/SymsLib/syms/core/syms_debug_info.h
--- a/Engine/Source/Runtime/SymsLib/syms/core/syms_debug_info.h
+++ b/Engine/Source/Runtime/SymsLib/syms/core/syms_debug_info.h
@@ -350,7 +350,7 @@ typedef struct SYMS_EnumInfoArray{
 ////////////////////////////////
 // NOTE(allen): Nil For Accelerators
 
-SYMS_GLOBAL SYMS_FileFormat syms_format_nil = SYMS_FileFormat_Null;
+SYMS_READ_ONLY SYMS_GLOBAL SYMS_FileFormat syms_format_nil = SYMS_FileFormat_Null;
 
 ////////////////////////////////
 //~ NOTE(allen): Symbol Helper Functions
@@ -365,6 +365,9 @@ SYMS_API SYMS_USID    syms_make_usid(SYMS_UnitID uid, SYMS_SymbolID sid);
 SYMS_API SYMS_TypeKind syms_type_kind_fwd_from_main(SYMS_TypeKind type_kind);
 SYMS_API SYMS_TypeKind syms_type_kind_main_from_fwd(SYMS_TypeKind type_kind);
 SYMS_API SYMS_B32      syms_type_kind_is_basic(SYMS_TypeKind kind);
+SYMS_API SYMS_B32      syms_type_kind_is_integer(SYMS_TypeKind kind);
+SYMS_API SYMS_B32      syms_type_kind_is_signed(SYMS_TypeKind kind);
+SYMS_API SYMS_B32      syms_type_kind_is_complex(SYMS_TypeKind kind);
 SYMS_API SYMS_B32      syms_type_kind_is_user_defined(SYMS_TypeKind kind);
 
 SYMS_API SYMS_SymbolIDArray syms_sid_array_from_list(SYMS_Arena *arena, SYMS_SymbolIDList *list);
diff --git a/Engine/Source/Runtime/SymsLib/syms/symslib.h b/Engine/Source/Runtime/SymsLib/syms/symslib.h
--- a/Engine/Source/Runtime/SymsLib/syms/symslib.h
+++ b/Engine/Source/Runtime/SymsLib/syms/symslib.h
@@ -8,6 +8,7 @@
 extern "C" {
 #endif
 
+#define SYMS_READ_ONLY
 #include "syms/overrides/syms_crt_overrides.h"
 #include "syms/syms_inc.h"
 
