- Updated Sysmlib that can handle > 4gb pdb files

diff --git a/Engine/Source/Runtime/SymsLib/syms/core/base/syms_base.h b/Engine/Source/Runtime/SymsLib/syms/core/base/syms_base.h
--- a/Engine/Source/Runtime/SymsLib/syms/core/base/syms_base.h
+++ b/Engine/Source/Runtime/SymsLib/syms/core/base/syms_base.h
@@ -7,9 +7,10 @@
 ////////////////////////////////
 //~ NOTE(allen): Version
 
-#define SYMS_VERSION_MAJOR  0
-#define SYMS_VERSION_MINOR  9999
-#define SYMS_VERSION_STR    "1.pre"
+#define SYMS_VERSION_MAJOR    1
+#define SYMS_VERSION_MINOR    0
+#define SYMS_VERSION_SUBMINOR 0
+#define SYMS_VERSION_STR      "1.0.0"
 
 ////////////////////////////////
 //~ NOTE(rjf): Context Cracking
@@ -41,7 +42,7 @@
 #define SYMS_GLOBAL static
 #define SYMS_LOCAL  static
 
-#ifndef SYMS_READ_ONLY
+#if !defined(SYMS_READ_ONLY)
 # if SYMS_COMPILER_CL || (SYMS_COMPILER_CLANG && SYMS_OS_WINDOWS)
 #  pragma section(".roglob", read)
 #  define SYMS_READ_ONLY __declspec(allocate(".roglob"))
@@ -183,6 +184,7 @@ typedef SYMS_U32 SYMS_RegID;
 #define SYMS_CeilIntegerDiv(a,b) (((a) + (b) - 1)/(b))
 
 #define SYMS_AlignPow2(a,b) (((a) + (b) - 1)&(~((b) - 1)))
+#define SYMS_AlignDownPow2(a,b) ((a)&(~((b) - 1)))
 
 #define SYMS_Swap(T,a,b) do{ T t__ = (a); (a) = (b); (b) = t__; }while(0)
 
@@ -227,6 +229,17 @@ typedef SYMS_U32 SYMS_RegID;
 ////////////////////////////////
 //~ NOTE(allen): Common Basic Types
 
+typedef struct SYMS_U64Array{
+  SYMS_U64 *u64;
+  SYMS_U64 count;
+} SYMS_U64Array;
+
+typedef struct SYMS_U64ArrayNode{
+  struct SYMS_U64ArrayNode *next;
+  SYMS_U64 *u64;
+  SYMS_U64 count;
+} SYMS_U64ArrayNode;
+
 typedef struct SYMS_U64Maybe{
   SYMS_B32 valid;
   SYMS_U64 u64;
@@ -416,7 +429,7 @@ SYMS_API SYMS_U64 syms_hash_u64(SYMS_U64 x);
 #define syms_serial_type(name) (_syms_serial_type_##name)
 #define syms_string_from_enum_value(enum_type, value) \
 (syms_serial_value_from_enum_value(&syms_serial_type(enum_type), value)->name)
-#define syms_bswap_in_place(type, ptr) syms_bswap_in_place__##type(ptr)
+#define syms_bswap_in_place(type, ptr) syms_bswap_in_place__##type((type*)(ptr))
 
 SYMS_API SYMS_SerialField* syms_serial_first_field(SYMS_SerialType *type);
 SYMS_API SYMS_SerialValue* syms_serial_first_value(SYMS_SerialType *type);
@@ -454,7 +467,7 @@ SYMS_API void         syms_string_list_push_node_front(SYMS_String8Node *node, S
                                                        SYMS_String8 string);
 SYMS_API void         syms_string_list_push(SYMS_Arena *arena, SYMS_String8List *list, SYMS_String8 string);
 SYMS_API void         syms_string_list_push_front(SYMS_Arena *arena, SYMS_String8List *list, SYMS_String8 string);
-SYMS_API SYMS_String8List syms_string_list_concat_in_place(SYMS_String8List *left, SYMS_String8List *right);
+SYMS_API SYMS_String8List syms_string_list_concat(SYMS_String8List *left, SYMS_String8List *right);
 
 
 SYMS_API SYMS_String8 syms_string_list_join(SYMS_Arena *arena, SYMS_String8List *list, SYMS_StringJoin *join);
diff --git a/Engine/Source/Runtime/SymsLib/syms/core/base/syms_base_overrides_check.h b/Engine/Source/Runtime/SymsLib/syms/core/base/syms_base_overrides_check.h
--- a/Engine/Source/Runtime/SymsLib/syms/core/base/syms_base_overrides_check.h
+++ b/Engine/Source/Runtime/SymsLib/syms/core/base/syms_base_overrides_check.h
@@ -11,120 +11,120 @@
 // NOTE(allen): basic types
 
 #if !defined(SYMS_S8)
-# error SYMS: SYMS_S8 was not #define'd. It should be #define'd to the name of a signed 8-bit integer type (int8_t in the C Standard Library).
+# error "SYMS: SYMS_S8 was not #define'd. It should be #define'd to the name of a signed 8-bit integer type (int8_t in the C Standard Library)."
 #endif
 
 #if !defined(SYMS_S16)
-# error SYMS: SYMS_S16 was not #define'd. It should be #define'd to the name of a signed 16-bit integer type (int16_t in the C Standard Library).
+# error "SYMS: SYMS_S16 was not #define'd. It should be #define'd to the name of a signed 16-bit integer type (int16_t in the C Standard Library)."
 #endif
 
 #if !defined(SYMS_S32)
-# error SYMS: SYMS_S32 was not #define'd. It should be #define'd to the name of a signed 32-bit integer type (int32_t in the C Standard Library).
+# error "SYMS: SYMS_S32 was not #define'd. It should be #define'd to the name of a signed 32-bit integer type (int32_t in the C Standard Library)."
 #endif
 
 #if !defined(SYMS_S64)
-# error SYMS: SYMS_S64 was not #define'd. It should be #define'd to the name of a signed 64-bit integer type (int64_t in the C Standard Library).
+# error "SYMS: SYMS_S64 was not #define'd. It should be #define'd to the name of a signed 64-bit integer type (int64_t in the C Standard Library)."
 #endif
 
 #if !defined(SYMS_U8)
-# error SYMS: SYMS_U8 was not #define'd. It should be #define'd to the name of an unsigned 8-bit integer type (uint8_t in the C Standard Library).
+# error "SYMS: SYMS_U8 was not #define'd. It should be #define'd to the name of an unsigned 8-bit integer type (uint8_t in the C Standard Library)."
 #endif
 
 #if !defined(SYMS_U16)
-# error SYMS: SYMS_U16 was not #define'd. It should be #define'd to the name of an unsigned 16-bit integer type (uint16_t in the C Standard Library).
+# error "SYMS: SYMS_U16 was not #define'd. It should be #define'd to the name of an unsigned 16-bit integer type (uint16_t in the C Standard Library)."
 #endif
 
 #if !defined(SYMS_U32)
-# error SYMS: SYMS_U32 was not #define'd. It should be #define'd to the name of an unsigned 32-bit integer type (uint32_t in the C Standard Library).
+# error "SYMS: SYMS_U32 was not #define'd. It should be #define'd to the name of an unsigned 32-bit integer type (uint32_t in the C Standard Library)."
 #endif
 
 #if !defined(SYMS_U64)
-# error SYMS: SYMS_U64 was not #define'd. It should be #define'd to the name of an unsigned 64-bit integer type (uint64_t in the C Standard Library).
+# error "SYMS: SYMS_U64 was not #define'd. It should be #define'd to the name of an unsigned 64-bit integer type (uint64_t in the C Standard Library)."
 #endif
 
 #if !defined(SYMS_PRIu64)
-# error SYMS: SYMS_PRIu64 was not #define'd. It should be #define'd to a string constant suitable for printing 64-bit unsigned integers in C format strings (PRIu64 in the C Standard Library).
+# error "SYMS: SYMS_PRIu64 was not #define'd. It should be #define'd to a string constant suitable for printing 64-bit unsigned integers in C format strings (PRIu64 in the C Standard Library)."
 #endif
 
 #if !defined(SYMS_PRId64)
-# error SYMS: SYMS_PRId64 was not #define'd. It should be #define'd to a string constant suitable for printing 64-bit signed integers in C format strings (PRId64 in the C Standard Library).
+# error "SYMS: SYMS_PRId64 was not #define'd. It should be #define'd to a string constant suitable for printing 64-bit signed integers in C format strings (PRId64 in the C Standard Library)."
 #endif
 
 #if !defined(SYMS_PRIx64)
-# error SYMS: SYMS_PRIx64 was not #define'd. It should be #define'd to a string constant suitable for printing 64-bit hex integers in C format strings (PRIx64 in the C Standard Library).
+# error "SYMS: SYMS_PRIx64 was not #define'd. It should be #define'd to a string constant suitable for printing 64-bit hex integers in C format strings (PRIx64 in the C Standard Library)."
 #endif
 
 
 // NOTE(allen): large data operations
 
 #if !defined(syms_memmove)
-# error SYMS: syms_memmove was not #define'd. It should be #define'd to the name of a memmove implementation.
+# error "SYMS: syms_memmove was not #define'd. It should be #define'd to the name of a memmove implementation."
 #endif
 
 #if !defined(syms_memset)
-# error SYMS: syms_memset was not #define'd. It should be #define'd to the name of a memset implementation.
+# error "SYMS: syms_memset was not #define'd. It should be #define'd to the name of a memset implementation."
 #endif
 
 #if !defined(syms_memcmp)
-# error SYMS: syms_memcmp was not #define'd. It should be #define'd to the name of a memcmp implementation.
+# error "SYMS: syms_memcmp was not #define'd. It should be #define'd to the name of a memcmp implementation."
 #endif
 
 #if !defined(syms_strlen)
-# error SYMS: syms_strlen was not #define'd. It should be #define'd to the name of a strlen implementation.
+# error "SYMS: syms_strlen was not #define'd. It should be #define'd to the name of a strlen implementation."
 #endif
 
 
 // NOTE(allen): assert
 
 #if !defined(SYMS_ASSERT_BREAK)
-# error SYMS: SYMS_ASSERT_BREAK(message) was not #define'd. It should be #define'd to something that you'd like expanded when an assertion failure occurs.
+# error "SYMS: SYMS_ASSERT_BREAK(message) was not #define'd. It should be #define'd to something that you'd like expanded when an assertion failure occurs."
 #endif
 
 
 // NOTE(allen): arena
 
 #if !defined(SYMS_Arena)
-# error SYMS: SYMS_Arena was not #define'd.
+# error "SYMS: SYMS_Arena was not #define'd."
 #endif
 
 #if !defined(syms_arena_alloc__impl)
-# error SYMS: syms_arena_alloc__impl was not #define'd.
+# error "SYMS: syms_arena_alloc__impl was not #define'd."
 #endif
 
 #if !defined(syms_arena_release__impl)
-# error SYMS: syms_arena_release__impl was not #define'd.
+# error "SYMS: syms_arena_release__impl was not #define'd."
 #endif
 
 #if !defined(syms_arena_get_pos__impl)
-# error SYMS: syms_arena_get_pos__impl was not #define'd.
+# error "SYMS: syms_arena_get_pos__impl was not #define'd."
 #endif
 
 #if !defined(syms_arena_push__impl)
-# error SYMS: syms_arena_push__impl was not #define'd.
+# error "SYMS: syms_arena_push__impl was not #define'd."
 #endif
 
 #if !defined(syms_arena_pop_to__impl)
-# error SYMS: syms_arena_pop_to__impl was not #define'd.
+# error "SYMS: syms_arena_pop_to__impl was not #define'd."
 #endif
 
 #if !defined(syms_arena_set_auto_align__impl)
-# error SYMS: syms_arena_set_auto_align__impl was not #define'd.
+# error "SYMS: syms_arena_set_auto_align__impl was not #define'd."
 #endif
 
 #if !defined(syms_arena_absorb__impl)
-# error SYMS: syms_arena_absorb__impl was not #define'd.
+# error "SYMS: syms_arena_absorb__impl was not #define'd."
 #endif
 
 #if !defined(syms_arena_tidy__impl)
-# error SYMS: syms_arena_tidy__impl was not #define'd.
+# error "SYMS: syms_arena_tidy__impl was not #define'd."
 #endif
 
 #if !defined(syms_get_implicit_thread_arena__impl)
-# error SYMS: syms_get_implicit_thread_arena__impl was not #define'd.
+# error "SYMS: syms_get_implicit_thread_arena__impl was not #define'd."
 #endif
 
 #if !defined(syms_scratch_pool_tidy__impl)
-# error SYMS: syms_scratch_pool_tidy__impl was not #define'd.
+# error "SYMS: syms_scratch_pool_tidy__impl was not #define'd."
 #endif
 
 #endif //SYMS_BASE_OVERRIDES_CHECK_H
diff --git a/Engine/Source/Runtime/SymsLib/syms/core/base/syms_dev.c b/Engine/Source/Runtime/SymsLib/syms/core/base/syms_dev.c
--- a/Engine/Source/Runtime/SymsLib/syms/core/base/syms_dev.c
+++ b/Engine/Source/Runtime/SymsLib/syms/core/base/syms_dev.c
@@ -1,4 +1,5 @@
 // Copyright Epic Games, Inc. All Rights Reserved.
+
 #ifndef SYMS_DEV_C
 #define SYMS_DEV_C
 
@@ -151,6 +152,15 @@ syms_logf__dev(char *fmt, ...){
   va_end(args);
 }
 
+#else
+
+#define syms_log_set_filter__dev(...)
+#define syms_log_open__dev(...)
+#define syms_log_is_enabled__dev(...) syms_false
+#define syms_log_open_annotated__dev(...) 0
+#define syms_logfv__dev(...)
+#define syms_logf__dev(...)
+
 #endif
 
 ////////////////////////////////
@@ -496,6 +506,28 @@ syms_prof_stringize_basic__dev(SYMS_Arena *arena, SYMS_ProfLock lock, SYMS_Strin
 
 #endif //SYMS_ENABLE_DEV_STRING
 
+#else
+
+#define syms_prof_equip_thread__dev(...)
+#define syms_prof_equip_thread_auto__dev(...)
+#define syms_prof_unequip_thread__dev(...)
+#define syms_prof_push__dev(...)
+#define syms_prof_paste__dev(...)
+
+#define syms_prof_lock__dev(...) {0}
+#define syms_prof_clear__dev(...)
+#define syms_prof_unlock__dev(l) ((void)l)
+
+#define syms_prof_tree__dev(...)
+#define syms_prof_tree_sort_in_place__dev(...)
+#define syms_prof_tree_sort_pointer_array__dev(...)
+
+#if SYMS_ENABLE_DEV_STRING
+# define syms_prof_stringize_tree__dev(...)
+# define syms_prof_stringize_tree__rec__dev(...)
+# define syms_prof_stringize_basic__dev(...)
+#endif
+
 #endif //SYMS_ENABLE_DEV_PROFILE
 
 #endif //SYMS_DEV_C
diff --git a/Engine/Source/Runtime/SymsLib/syms/core/base/syms_dev.h b/Engine/Source/Runtime/SymsLib/syms/core/base/syms_dev.h
--- a/Engine/Source/Runtime/SymsLib/syms/core/base/syms_dev.h
+++ b/Engine/Source/Runtime/SymsLib/syms/core/base/syms_dev.h
@@ -150,12 +150,12 @@ SYMS_API void syms_logf__dev(char *fmt, ...);
 
 // SYMS_PROF_TIME: () -> SYMS_U64
 #if !defined(SYMS_PROF_TIME)
-# error SYMS_PROF_TIME not #define'd
+# error "SYMS_PROF_TIME not #define'd"
 #endif
 
 // SYMS_PROF_ALLOC: (SYMS_U64) -> void*
 #if !defined(SYMS_PROF_ALLOC)
-# error SYMS_PROF_ALLOC not #define'd
+# error "SYMS_PROF_ALLOC not #define'd"
 #endif
 
 #if !defined(SYMS_PROF_BLOCK_SIZE)
diff --git a/Engine/Source/Runtime/SymsLib/syms/core/data_structures/syms_data_structures.c b/Engine/Source/Runtime/SymsLib/syms/core/data_structures/syms_data_structures.c
--- a/Engine/Source/Runtime/SymsLib/syms/core/data_structures/syms_data_structures.c
+++ b/Engine/Source/Runtime/SymsLib/syms/core/data_structures/syms_data_structures.c
@@ -1,4 +1,5 @@
 // Copyright Epic Games, Inc. All Rights Reserved.
+
 #if !defined(SYMS_DATA_STRUCTURES_C)
 #define SYMS_DATA_STRUCTURES_C
 
@@ -17,7 +18,7 @@ SYMS_API SYMS_String8
 syms_string_cons(SYMS_Arena *arena, SYMS_StringCons *cons, SYMS_String8 string){
   SYMS_ProfBegin("syms_string_cons");
   SYMS_String8 result = {0};
-  if (string.size > 0){
+  if (cons->bucket_count > 0 && string.size > 0){
     SYMS_U64 hash = syms_hash_djb2(string);
     SYMS_U64 bucket_index = hash%cons->bucket_count;
     for (SYMS_StringConsNode *node = cons->buckets[bucket_index];
@@ -43,6 +44,134 @@ syms_string_cons(SYMS_Arena *arena, SYMS_StringCons *cons, SYMS_String8 string){
   return(result);
 }
 
+
+SYMS_API SYMS_DataIdxCons
+syms_data_idx_cons_alloc(SYMS_Arena *arena, SYMS_U64 bucket_count){
+  SYMS_DataIdxCons result = {0};
+  result.bucket_count = bucket_count;
+  result.buckets = syms_push_array_zero(arena, SYMS_DataIdxConsNode*, bucket_count);
+  return(result);
+}
+
+SYMS_API SYMS_U64
+syms_data_idx_cons(SYMS_Arena *arena, SYMS_DataIdxCons *cons, SYMS_String8 data){
+  SYMS_ProfBegin("syms_data_idx_cons");
+  SYMS_U64 result = 0;
+  if (cons->bucket_count > 0 && data.size > 0){
+    SYMS_U64 hash = syms_hash_djb2(data);
+    SYMS_U64 bucket_index = hash%cons->bucket_count;
+    for (SYMS_DataIdxConsNode *node = cons->buckets[bucket_index];
+         node != 0;
+         node = node->bucket_next){
+      if (node->hash == hash && syms_string_match(data, node->data, 0)){
+        result = node->id;
+        break;
+      }
+    }
+    if (result == 0){
+      // take a new id
+      SYMS_U64 new_id = cons->count + 1;
+      // stabilize the data memory
+      SYMS_String8 stable_data = syms_push_string_copy(arena, data);
+      // save a cons node
+      SYMS_DataIdxConsNode *new_node = syms_push_array_zero(arena, SYMS_DataIdxConsNode, 1);
+      SYMS_StackPush_N(cons->buckets[bucket_index], new_node, bucket_next);
+      SYMS_QueuePush_N(cons->first, cons->last, new_node, all_next);
+      cons->count += 1;
+      cons->total_size += data.size;
+      new_node->data = stable_data;
+      new_node->hash = hash;
+      new_node->id = new_id;
+      result = new_id;
+    }
+  }
+  SYMS_ProfEnd();
+  return(result);
+}
+
+
+////////////////////////////////
+//~ NOTE(allen): U64 Set
+
+SYMS_API SYMS_U64Set
+syms_u64_set_alloc(SYMS_Arena *arena, SYMS_U64 cap){
+  SYMS_U64Set result = {0};
+  result.cap = cap;
+  result.vals = syms_push_array_zero(arena, SYMS_U64, cap);
+  return(result);
+}
+
+SYMS_API SYMS_U64
+syms_u64_set__bs(SYMS_U64Set *set, SYMS_U64 x){
+  SYMS_ProfBegin("syms_u64_set__bs");
+  //- binary search:
+  //  min index s.t. x <= vals[index]
+  //  set->count if no such index exists
+  // in this one we assume:
+  //  (i != j) implies (vals[i] != vals[j])
+  //  thus if (vals[index] == x) then index already satisfies the requirement
+  SYMS_U64 count = set->count;
+  SYMS_U64 *vals = set->vals;
+  SYMS_U64 result = set->count;
+  if (count > 0 && x <= vals[count - 1]){
+    SYMS_U64 first = 0;
+    SYMS_U64 opl   = count;
+    for (;;){
+      if (first + 1 >= opl){
+        break;
+      }
+      SYMS_U64 mid = (first + opl - 1)/2;
+      SYMS_U64 v = vals[mid];
+      if (x < v){
+        opl = mid + 1;
+      }
+      else if (x > v){
+        first = mid + 1;
+      }
+      else{
+        first = mid;
+        break;
+      }
+    }
+    result = first;
+  }
+  SYMS_ProfEnd();
+  return(result);
+}
+
+SYMS_API SYMS_B32
+syms_u64_set_insert(SYMS_U64Set *set, SYMS_U64 x){
+  SYMS_ProfBegin("syms_u64_set_insert");
+  SYMS_B32 result = 0;
+  SYMS_U64 index = syms_u64_set__bs(set, x);
+  SYMS_U64 *vals = set->vals;
+  SYMS_U64 count = set->count;
+  if (index < set->cap && (index == count || vals[index] != x)){
+    syms_memmove(vals + index + 1, vals + index,
+                 sizeof(SYMS_U64)*(count - index));
+    vals[index] = x;
+    set->count = count + 1;
+    result = 1;
+  }
+  SYMS_ProfEnd();
+  return(result);
+}
+
+SYMS_API void
+syms_u64_set_erase(SYMS_U64Set *set, SYMS_U64 x){
+  SYMS_ProfBegin("syms_u64_set_erase");
+  SYMS_U64 index = syms_u64_set__bs(set, x);
+  SYMS_U64 *vals = set->vals;
+  SYMS_U64 count = set->count;
+  if (index < count && vals[index] == x){
+    syms_memmove(vals + index, vals + index + 1,
+                 sizeof(SYMS_U64)*(count - index - 1));
+    set->count = count - 1;
+  }
+  SYMS_ProfEnd();
+}
+
+
 ////////////////////////////////
 //~ NOTE(allen): Syms Spatial Mapping
 
@@ -65,10 +194,10 @@ syms_spatial_map_1d_binary_search(SYMS_SpatialMap1D *map, SYMS_U64 x){
     for (;;){
       SYMS_U64 mid = (first + opl)/2;
       SYMS_U64 rmin = ranges[mid].range.min;
-      if (rmin > x){
+      if (x < rmin){
         opl = mid;
       }
-      else if (rmin < x){
+      else if (x > rmin){
         first = mid;
       }
       else{
@@ -285,6 +414,250 @@ syms_spatial_map_1d_array_sort__rec(SYMS_SpatialMap1DRange *ranges, SYMS_U64 cou
   }
 }
 
+//- support for the overlapping ranges
+
+SYMS_API SYMS_SpatialMultiMap1D
+syms_spatial_multi_map_1d_bake(SYMS_Arena *arena, SYMS_SpatialMap1DLoose *loose){
+  SYMS_ProfBegin("syms_spatial_multi_map_1d_bake");
+  SYMS_SpatialMultiMap1D result = {0};
+  
+  if (loose->total_count > 0){
+    //- gather endpoints
+    SYMS_ArenaTemp scratch = syms_get_scratch(&arena, 1);
+    SYMS_U64 max_endpoint_count = loose->total_count*2;
+    SYMS_1DEndPoint *endpoints = syms_push_array(scratch.arena, SYMS_1DEndPoint, max_endpoint_count);
+    
+    SYMS_U64 endpoint_count = 0;
+    {
+      SYMS_1DEndPoint *endpoint_ptr = endpoints;
+      for (SYMS_SpatialMap1DNode *node = loose->first;
+           node != 0;
+           node = node->next){
+        SYMS_U64 val = node->val;
+        
+        {
+          SYMS_U64Range range = node->range;
+          if (range.min < range.max){
+            endpoint_ptr->x = range.min;
+            endpoint_ptr->val = val;
+            endpoint_ptr->open = 1;
+            endpoint_ptr += 1;
+            
+            endpoint_ptr->x = range.max;
+            endpoint_ptr->val = val;
+            endpoint_ptr->open = 0;
+            endpoint_ptr += 1;
+          }
+        }
+        
+        SYMS_U64 count = node->ranges.count;
+        SYMS_U64Range *range = node->ranges.ranges;
+        for (SYMS_U64 i = 0; i < count; i += 1, range += 1){
+          if (range->min < range->max){
+            endpoint_ptr->x = range->min;
+            endpoint_ptr->val = val;
+            endpoint_ptr->open = 1;
+            endpoint_ptr += 1;
+            
+            endpoint_ptr->x = range->max;
+            endpoint_ptr->val = val;
+            endpoint_ptr->open = 0;
+            endpoint_ptr += 1;
+          }
+        }
+      }
+      
+      endpoint_count = (SYMS_U64)(endpoint_ptr - endpoints);
+    }
+    SYMS_U64 range_count = endpoint_count/2;
+    
+    //- sort endpoints
+    syms_spatial_map_1d_endpoint_sort(endpoints, endpoint_count);
+    
+    //- walk endpoints
+    SYMS_SpatialMap1D spatial_map = {0};
+    SYMS_DataIdxCons data_cons = syms_data_idx_cons_alloc(scratch.arena, loose->total_count);
+    if (endpoint_count > 0){
+      // allocate simple map ranges
+      SYMS_SpatialMap1DRange *ranges = syms_push_array(arena, SYMS_SpatialMap1DRange, range_count);
+      
+      // setup the walk
+      SYMS_SpatialMap1DRange *range_ptr = ranges;
+      SYMS_1DEndPoint *endpoint_ptr = endpoints;
+      SYMS_1DEndPoint *endpoint_opl = endpoints + endpoint_count;
+      SYMS_U64Set set = syms_u64_set_alloc(scratch.arena, range_count);
+      
+      // walk
+      SYMS_U64 prev_x = 0;
+      SYMS_B32 has_prev = 0;
+      for (;endpoint_ptr < endpoint_opl;){
+        // look at the upcoming x
+        SYMS_U64 this_x = endpoint_ptr->x;
+        
+        // emit a range
+        if (has_prev && set.count > 0){
+          // get an index for this set
+          SYMS_String8 data = {0};
+          data.str = (SYMS_U8*)set.vals;
+          data.size = set.count*sizeof(*set.vals);
+          SYMS_U64 set_number = syms_data_idx_cons(scratch.arena, &data_cons, data);
+          
+          // fill range and advance
+          range_ptr->range.min = prev_x;
+          range_ptr->range.max = this_x;
+          range_ptr->val = set_number;
+          range_ptr += 1;
+        }
+        
+        // consume all endpoints at this x
+        for (;endpoint_ptr < endpoint_opl && endpoint_ptr->x == this_x;
+             endpoint_ptr += 1){
+          if (endpoint_ptr->open){
+            syms_u64_set_insert(&set, endpoint_ptr->val);
+          }
+          else{
+            syms_u64_set_erase(&set, endpoint_ptr->val);
+          }
+        }
+        
+        // update prev info
+        prev_x = this_x;
+        has_prev = 1;
+      }
+      
+      // since we control emitting endpoints and consuming them,
+      // if we have done everything right then the walk should end
+      // with an empty set
+      SYMS_ASSERT(set.count == 0);
+      
+      // assemble the simple spatial map
+      spatial_map.ranges = ranges;
+      spatial_map.count = (SYMS_U64)(range_ptr - ranges);
+    }
+    
+    //- assemble the set data
+    SYMS_U64 set_count = data_cons.count;
+    SYMS_U64 *set_end_points = syms_push_array(arena, SYMS_U64, set_count + 1);
+    SYMS_U8 *set_data = syms_push_array(arena, SYMS_U8, data_cons.total_size);
+    {
+      // setup cursors
+      SYMS_U64 *set_end_point_ptr = set_end_points;
+      SYMS_U64 set_data_pos = 0;
+      
+      // emit loop
+      *set_end_point_ptr = 0;
+      set_end_point_ptr += 1;
+      for (SYMS_DataIdxConsNode *node = data_cons.first;
+           node != 0;
+           node = node->all_next){
+        SYMS_U64 data_size = node->data.size;
+        syms_memmove(set_data + set_data_pos, node->data.str, data_size);
+        set_data_pos += data_size;
+        *set_end_point_ptr = set_data_pos;
+        set_end_point_ptr += 1;
+      }
+    }
+    
+    //- assemble the multi-map
+    result.spatial_map = spatial_map;
+    result.set_end_points = set_end_points;
+    result.set_data = set_data;
+    result.set_count = set_count;
+    
+    syms_release_scratch(scratch);
+  }
+  
+  SYMS_ProfEnd();
+  return(result);
+}
+
+SYMS_API SYMS_U64Array
+syms_spatial_multi_map_1d_array_from_point(SYMS_SpatialMultiMap1D *map, SYMS_U64 x){
+  SYMS_ProfBegin("syms_spatial_multi_map_array_from_point");
+  SYMS_U64Array result = {0};
+  SYMS_U64 set_number = syms_spatial_map_1d_value_from_point(&map->spatial_map, x);
+  if (0 < set_number && set_number <= map->set_count){
+    SYMS_U64 *end_points = map->set_end_points;
+    SYMS_U64 set_index = set_number - 1;
+    SYMS_U64 first = end_points[set_index];
+    SYMS_U64 opl = end_points[set_index + 1];
+    result.u64 = (SYMS_U64*)(map->set_data + first);
+    result.count = (opl - first)/8;
+  }
+  SYMS_ProfEnd();
+  return(result);
+}
+
+
+SYMS_API void
+syms_spatial_map_1d_endpoint_sort(SYMS_1DEndPoint *endpoints, SYMS_U64 count){
+  if (count > 4){
+    SYMS_U64 last = count - 1;
+    
+    // swap
+    SYMS_U64 mid = count/2;
+    SYMS_Swap(SYMS_1DEndPoint, endpoints[mid], endpoints[last]);
+    
+    // partition
+    SYMS_B32 equal_send_left = 0;
+    SYMS_U64 key = endpoints[last].x;
+    SYMS_U64 j = 0;
+    for (SYMS_U64 i = 0; i < last; i += 1){
+      SYMS_B32 send_left = (endpoints[i].x < key);
+      if (!send_left && endpoints[i].x == key){
+        send_left = equal_send_left;
+        equal_send_left = !equal_send_left;
+      }
+      if (send_left){
+        if (j != i){
+          SYMS_Swap(SYMS_1DEndPoint, endpoints[i], endpoints[j]);
+        }
+        j += 1;
+      }
+    }
+    
+    SYMS_Swap(SYMS_1DEndPoint, endpoints[j], endpoints[last]);
+    
+    // recurse
+    SYMS_U64 pivot = j;
+    syms_spatial_map_1d_endpoint_sort(endpoints, pivot);
+    syms_spatial_map_1d_endpoint_sort(endpoints + pivot + 1, (count - pivot - 1));
+  }
+  else if (count == 2){
+    if (endpoints[0].x > endpoints[1].x){
+      SYMS_Swap(SYMS_1DEndPoint, endpoints[0], endpoints[1]);
+    }
+  }
+  else if (count == 3){
+    if (endpoints[0].x > endpoints[1].x){
+      SYMS_Swap(SYMS_1DEndPoint, endpoints[0], endpoints[1]);
+    }
+    if (endpoints[1].x > endpoints[2].x){
+      SYMS_Swap(SYMS_1DEndPoint, endpoints[1], endpoints[2]);
+      if (endpoints[0].x > endpoints[1].x){
+        SYMS_Swap(SYMS_1DEndPoint, endpoints[0], endpoints[1]);
+      }
+    }
+  }
+  else if (count == 4){
+    if (endpoints[0].x > endpoints[1].x){
+      SYMS_Swap(SYMS_1DEndPoint, endpoints[0], endpoints[1]);
+    }
+    if (endpoints[2].x > endpoints[3].x){
+      SYMS_Swap(SYMS_1DEndPoint, endpoints[2], endpoints[3]);
+    }
+    if (endpoints[0].x > endpoints[2].x){
+      SYMS_Swap(SYMS_1DEndPoint, endpoints[0], endpoints[2]);
+    }
+    if (endpoints[1].x > endpoints[3].x){
+      SYMS_Swap(SYMS_1DEndPoint, endpoints[1], endpoints[3]);
+    }
+    if (endpoints[1].x > endpoints[2].x){
+      SYMS_Swap(SYMS_1DEndPoint, endpoints[1], endpoints[2]);
+    }
+  }
+}
+
 //- invariants for spatial maps
 
 SYMS_API SYMS_B32
@@ -642,6 +1015,141 @@ syms_id_map_insert(SYMS_Arena *arena, SYMS_IDMap *map, SYMS_U64 key, void *val){
 }
 
 
+////////////////////////////////
+//~ NOTE(allen): Symbol Name Mapping Structure (String -> Array(USID))
+
+//- lookups into symbol name map
+
+SYMS_API SYMS_SymbolIDArray
+syms_symbol_name_map_array_from_string(SYMS_SymbolNameMap *map, SYMS_String8 string){
+  SYMS_ProfBegin("syms_symbol_name_map_array_from_string");
+  SYMS_SymbolIDArray result = {0};
+  if (map->bucket_count > 0){
+    SYMS_U64 hash = syms_hash_djb2(string);
+    SYMS_U64 bucket_index = hash%map->bucket_count;
+    for (SYMS_SymbolNameNode *node = map->buckets[bucket_index];
+         node != 0;
+         node = node->next_bucket){
+      if (node->hash == hash && syms_string_match(node->name, string, 0)){
+        result = node->sid_array;
+        break;
+      }
+    }
+  }
+  SYMS_ProfEnd();
+  return(result);
+}
+
+//- constructing symbol name maps
+
+SYMS_API SYMS_SymbolNameMapLoose
+syms_symbol_name_map_begin(SYMS_Arena *arena, SYMS_U64 bucket_count){
+  SYMS_SymbolNameMapLoose result = {0};
+  result.bucket_count = bucket_count;
+  result.buckets = syms_push_array_zero(arena, SYMS_SymbolNameNodeLoose*, bucket_count);
+  return(result);
+}
+
+SYMS_API void
+syms_symbol_name_map_push(SYMS_Arena *arena, SYMS_SymbolNameMapLoose *map,
+                          SYMS_String8 name, SYMS_SymbolID sid){
+  SYMS_ProfBegin("syms_symbol_name_map_push");
+  if (map->bucket_count > 0){
+    SYMS_U64 hash = syms_hash_djb2(name);
+    SYMS_U64 bucket_index = hash%map->bucket_count;
+    
+    // find existing node
+    SYMS_SymbolNameNodeLoose *node = 0;
+    for (SYMS_SymbolNameNodeLoose *bucket = map->buckets[bucket_index];
+         bucket != 0;
+         bucket = bucket->next_bucket){
+      if (bucket->hash == hash &&
+          syms_string_match(name, bucket->name, 0)){
+        node = bucket;
+        break;
+      }
+    }
+    
+    // create node
+    if (node == 0){
+      node = syms_push_array_zero(arena, SYMS_SymbolNameNodeLoose, 1);
+      SYMS_StackPush_N(map->buckets[bucket_index], node, next_bucket);
+      SYMS_QueuePush(map->first, map->last, node);
+      map->node_count += 1;
+      node->name = syms_push_string_copy(arena, name);
+      node->hash = hash;
+    }
+    
+    // insert sid
+    {
+      SYMS_SymbolIDNode *sid_node = syms_push_array_zero(arena, SYMS_SymbolIDNode, 1);
+      SYMS_QueuePush(node->sid_list.first, node->sid_list.last, sid_node);
+      node->sid_list.count += 1;
+      sid_node->id = sid;
+    }
+  }
+  SYMS_ProfEnd();
+}
+
+SYMS_API SYMS_SymbolNameMap
+syms_symbol_name_map_bake(SYMS_Arena *arena, SYMS_SymbolNameMapLoose *loose){
+  SYMS_ProfBegin("syms_symbol_name_map_bake");
+  
+  //- allocate map memory
+  SYMS_U64 node_count = loose->node_count;
+  SYMS_SymbolNameNode *nodes = syms_push_array(arena, SYMS_SymbolNameNode, node_count);
+  
+  SYMS_U64 bucket_count = node_count*3/2;
+  SYMS_SymbolNameNode **buckets = syms_push_array_zero(arena, SYMS_SymbolNameNode*, bucket_count);
+  
+  //- fill map from loose
+  SYMS_SymbolNameNode *node_ptr = nodes;
+  for (SYMS_SymbolNameNodeLoose *loose_node = loose->first;
+       loose_node != 0;
+       loose_node = loose_node->next){
+    // copy node's name
+    SYMS_String8 name = syms_push_string_copy(arena, loose_node->name);
+    
+    // grab the node's hash
+    SYMS_U64 hash = loose_node->hash;
+    
+    // fill node's usid array
+    SYMS_SymbolIDArray sid_array = {0};
+    sid_array.count = loose_node->sid_list.count;
+    sid_array.ids = syms_push_array(arena, SYMS_SymbolID, sid_array.count);
+    {
+      SYMS_SymbolID *sid_ptr = sid_array.ids;
+      for (SYMS_SymbolIDNode *sid_node = loose_node->sid_list.first;
+           sid_node != 0;
+           sid_node = sid_node->next){
+        *sid_ptr = sid_node->id;
+        sid_ptr += 1;
+      }
+    }
+    
+    // insert node to bucket
+    SYMS_U64 bucket_index = hash%bucket_count;
+    SYMS_StackPush_N(buckets[bucket_index], node_ptr, next_bucket);
+    
+    // fill node ptr
+    node_ptr->name = name;
+    node_ptr->hash = hash;
+    node_ptr->sid_array = sid_array;
+    node_ptr += 1;
+  }
+  
+  //- assemble result
+  SYMS_SymbolNameMap result = {0};
+  result.buckets = buckets;
+  result.bucket_count = bucket_count;
+  result.nodes = nodes;
+  result.node_count = node_count;
+  
+  SYMS_ProfEnd();
+  return(result);
+}
+
+
 ////////////////////////////////
 //~ NOTE(allen): Line Tables
 
diff --git a/Engine/Source/Runtime/SymsLib/syms/core/data_structures/syms_data_structures.h b/Engine/Source/Runtime/SymsLib/syms/core/data_structures/syms_data_structures.h
--- a/Engine/Source/Runtime/SymsLib/syms/core/data_structures/syms_data_structures.h
+++ b/Engine/Source/Runtime/SymsLib/syms/core/data_structures/syms_data_structures.h
@@ -20,6 +20,34 @@ typedef struct SYMS_StringCons{
   SYMS_U64 bucket_count;
 } SYMS_StringCons;
 
+// assign sequential indexes to small variable length blobs
+
+typedef struct SYMS_DataIdxConsNode{
+  struct SYMS_DataIdxConsNode *bucket_next;
+  struct SYMS_DataIdxConsNode *all_next;
+  SYMS_String8 data;
+  SYMS_U64 hash;
+  SYMS_U64 id;
+} SYMS_DataIdxConsNode;
+
+typedef struct SYMS_DataIdxCons{
+  SYMS_DataIdxConsNode **buckets;
+  SYMS_U64 bucket_count;
+  SYMS_DataIdxConsNode *first;
+  SYMS_DataIdxConsNode *last;
+  SYMS_U64 count;
+  SYMS_U64 total_size;
+} SYMS_DataIdxCons;
+
+////////////////////////////////
+//~ NOTE(allen): U64 Set
+
+typedef struct SYMS_U64Set{
+  SYMS_U64 *vals;
+  SYMS_U64 count;
+  SYMS_U64 cap;
+} SYMS_U64Set;
+
 ////////////////////////////////
 //~ NOTE(allen): 1D Spatial Mapping Structure
 
@@ -32,6 +60,7 @@ typedef struct SYMS_SpatialMap1DRange{
   SYMS_U64 val;
 } SYMS_SpatialMap1DRange;
 
+// TODO(allen): optimize this by using end points instead of ranges
 typedef struct SYMS_SpatialMap1D{
   SYMS_SpatialMap1DRange *ranges;
   SYMS_U64 count;
@@ -51,6 +80,22 @@ typedef struct SYMS_SpatialMap1DLoose{
   SYMS_U64 total_count;
 } SYMS_SpatialMap1DLoose;
 
+//- version with support for overlapping ranges
+typedef struct SYMS_SpatialMultiMap1D{
+  SYMS_SpatialMap1D spatial_map;
+  // set_end_points length: (set_count + 1)
+  SYMS_U64 *set_end_points; 
+  // set_data length: sizeof(SYMS_U64)*(set_end_points[set_count])
+  SYMS_U8 *set_data;
+  SYMS_U64 set_count;
+} SYMS_SpatialMultiMap1D;
+
+typedef struct SYMS_1DEndPoint{
+  SYMS_U64 x;
+  SYMS_U64 val;
+  SYMS_B32 open;
+} SYMS_1DEndPoint;
+
 ////////////////////////////////
 //~ NOTE(allen): File Mapping Structure ({UnitID,FileID} -> String)
 
@@ -144,15 +189,64 @@ typedef struct SYMS_IDMap{
 } SYMS_IDMap;
 
 
+////////////////////////////////
+//~ NOTE(allen): Symbol Name Mapping Structure (String -> Array(USID))
+
+// maps strings to lists of USIDs
+// organized as an array of nodes and a hash table simultaneously
+// so that single name lookups are accelerated, and filter matching
+// is also possible.
+
+typedef struct SYMS_SymbolNameNode{
+  struct SYMS_SymbolNameNode *next_bucket;
+  SYMS_String8 name;
+  SYMS_U64 hash;
+  SYMS_SymbolIDArray sid_array;
+} SYMS_SymbolNameNode;
+
+typedef struct SYMS_SymbolNameMap{
+  SYMS_SymbolNameNode **buckets;
+  SYMS_U64 bucket_count;
+  SYMS_SymbolNameNode *nodes;
+  SYMS_U64 node_count;
+} SYMS_SymbolNameMap;
+
+//- loose version
+typedef struct SYMS_SymbolNameNodeLoose{
+  struct SYMS_SymbolNameNodeLoose *next;
+  struct SYMS_SymbolNameNodeLoose *next_bucket;
+  SYMS_String8 name;
+  SYMS_U64 hash;
+  SYMS_SymbolIDList sid_list;
+} SYMS_SymbolNameNodeLoose;
+
+typedef struct SYMS_SymbolNameMapLoose{
+  SYMS_SymbolNameNodeLoose **buckets;
+  SYMS_U64 bucket_count;
+  SYMS_SymbolNameNodeLoose *first;
+  SYMS_SymbolNameNodeLoose *last;
+  SYMS_U64 node_count;
+} SYMS_SymbolNameMapLoose;
+
 ////////////////////////////////
 //~ NOTE(allen): String Cons Functions
 
 SYMS_API SYMS_StringCons syms_string_cons_alloc(SYMS_Arena *arena, SYMS_U64 bucket_count);
 SYMS_API SYMS_String8    syms_string_cons(SYMS_Arena *arena, SYMS_StringCons *cons, SYMS_String8 string);
 
+SYMS_API SYMS_DataIdxCons syms_data_idx_cons_alloc(SYMS_Arena *arena, SYMS_U64 bucket_count);
+SYMS_API SYMS_U64         syms_data_idx_cons(SYMS_Arena *arena, SYMS_DataIdxCons *cons, SYMS_String8 data);
 
 ////////////////////////////////
-//~ NOTE(allen): 1D Spatial Mapping Functions
+//~ NOTE(allen): U64 Set
+
+SYMS_API SYMS_U64Set syms_u64_set_alloc(SYMS_Arena *arena, SYMS_U64 cap);
+SYMS_API SYMS_U64    syms_u64_set__bs(SYMS_U64Set *set, SYMS_U64 x);
+SYMS_API SYMS_B32    syms_u64_set_insert(SYMS_U64Set *set, SYMS_U64 x);
+SYMS_API void        syms_u64_set_erase(SYMS_U64Set *set, SYMS_U64 x);
+
+////////////////////////////////
+//~ NOTE(allen): 1D Spatial Mapping Functions (Overlaps Not Allowed)
 
 //- lookups into spatial maps
 SYMS_API SYMS_U64          syms_spatial_map_1d_binary_search(SYMS_SpatialMap1D *map, SYMS_U64 x);
@@ -173,10 +267,18 @@ SYMS_API SYMS_B32          syms_spatial_map_1d_array_check_sorted(SYMS_SpatialMa
 SYMS_API void              syms_spatial_map_1d_array_sort(SYMS_SpatialMap1DRange *ranges, SYMS_U64 count);
 SYMS_API void              syms_spatial_map_1d_array_sort__rec(SYMS_SpatialMap1DRange *ranges, SYMS_U64 count);
 
+//- support for the overlapping ranges
+SYMS_API SYMS_SpatialMultiMap1D syms_spatial_multi_map_1d_bake(SYMS_Arena *arena, SYMS_SpatialMap1DLoose *loose);
+
+SYMS_API SYMS_U64Array syms_spatial_multi_map_1d_array_from_point(SYMS_SpatialMultiMap1D *map, SYMS_U64 x);
+
+SYMS_API void              syms_spatial_map_1d_endpoint_sort(SYMS_1DEndPoint *endpoints, SYMS_U64 count);
+
+// TODO(allen): copying spatial multi maps
+
 //- invariants for spatial maps
 SYMS_API SYMS_B32          syms_spatial_map_1d_invariants(SYMS_SpatialMap1D *map);
 
-
 ////////////////////////////////
 //~ NOTE(allen): File Mapping Functions ({UnitID,FileID} -> String)
 
@@ -227,6 +329,22 @@ SYMS_API SYMS_IDMap   syms_id_map_alloc(SYMS_Arena *arena, SYMS_U64 bucket_count
 SYMS_API void         syms_id_map_insert(SYMS_Arena *arena, SYMS_IDMap *map, SYMS_U64 key, void *val);
 
 
+////////////////////////////////
+//~ NOTE(allen): Symbol Name Mapping Structure (String -> Array(USID))
+
+// TODO(allen): copying
+
+//- lookups into symbol name map
+SYMS_API SYMS_SymbolIDArray syms_symbol_name_map_array_from_string(SYMS_SymbolNameMap *map, SYMS_String8 string);
+
+//- constructing symbol name maps
+SYMS_API SYMS_SymbolNameMapLoose syms_symbol_name_map_begin(SYMS_Arena *arena, SYMS_U64 bucket_count);
+SYMS_API void                    syms_symbol_name_map_push(SYMS_Arena *arena, SYMS_SymbolNameMapLoose *map,
+                                                           SYMS_String8 name, SYMS_SymbolID sid);
+SYMS_API SYMS_SymbolNameMap      syms_symbol_name_map_bake(SYMS_Arena *arena, SYMS_SymbolNameMapLoose *loose);
+
+
+
 ////////////////////////////////
 //~ NOTE(allen): Line Tables
 
@@ -241,7 +359,6 @@ SYMS_API void              syms_line_table_rewrite_file_ids_in_place(SYMS_FileID
 SYMS_API SYMS_LineTable    syms_line_table_with_indexes_from_parse(SYMS_Arena *arena, SYMS_LineParseOut *parse);
 
 
-
 ////////////////////////////////
 //~ NOTE(allen): Copies & Operators for Other Data Structures
 
diff --git a/Engine/Source/Runtime/SymsLib/syms/core/dwarf/syms_dwarf.c b/Engine/Source/Runtime/SymsLib/syms/core/dwarf/syms_dwarf.c
--- a/Engine/Source/Runtime/SymsLib/syms/core/dwarf/syms_dwarf.c
+++ b/Engine/Source/Runtime/SymsLib/syms/core/dwarf/syms_dwarf.c
@@ -1,4 +1,5 @@
 // Copyright Epic Games, Inc. All Rights Reserved.
+
 #ifndef SYMS_DWARF_C
 #define SYMS_DWARF_C
 
@@ -33,7 +34,8 @@ syms_dw_name_string_from_section_kind(SYMS_DwSectionKind kind)
     ".debug_loclists",
     ".debug_rnglists",
     ".debug_str_offsets",
-    ".debug_line_str"
+    ".debug_line_str",
+    ".debug_names"
   };
   SYMS_String8 str =
   {
@@ -56,15 +58,16 @@ syms_dw_mach_name_string_from_section_kind(SYMS_DwSectionKind kind)
     "__debug_line",
     "__debug_loc",
     "__debug_macinfo",
-    "__debug_pubnames",
-    "__debug_pubtypes",
+    "__debug_pubnames_DWARF",
+    "__debug_pubtypes_DWARF",
     "__debug_ranges",
     "__debug_str",
     "__debug_addr",
     "__debug_loclists",
     "__debug_rnglists",
     "__debug_str_offsets",
-    "__debug_line_str"
+    "__debug_line_str",
+    "__debug_names"
   };
   SYMS_String8 str =
   {
@@ -95,7 +98,8 @@ syms_dw_dwo_name_string_from_section_kind(SYMS_DwSectionKind kind)
     ".debug_loclists.dwo",
     ".debug_rnglists.dwo",
     ".debug_str_offsets.dwo",
-    ".debug_line_str.dwo"
+    ".debug_line_str.dwo",
+    ".debug_names.dwo"
   };
   SYMS_String8 str =
   {
diff --git a/Engine/Source/Runtime/SymsLib/syms/core/dwarf/syms_dwarf_expr.c b/Engine/Source/Runtime/SymsLib/syms/core/dwarf/syms_dwarf_expr.c
--- a/Engine/Source/Runtime/SymsLib/syms/core/dwarf/syms_dwarf_expr.c
+++ b/Engine/Source/Runtime/SymsLib/syms/core/dwarf/syms_dwarf_expr.c
@@ -1,4 +1,5 @@
 // Copyright Epic Games, Inc. All Rights Reserved.
+
 #if !defined(SYMS_DWARF_EXPR_C)
 #define SYMS_DWARF_EXPR_C
 
@@ -847,9 +848,9 @@ syms_dw_expr__eval(SYMS_Arena *arena_optional, void *expr_base, SYMS_U64Range ex
         
         case SYMS_DwOp_DEREF_SIZE:
         {
-          SYMS_U64 size = 0;
-          step_cursor += syms_based_range_read(base, range, step_cursor, 1, &size);
-          size = SYMS_ClampTop(size, 8);
+          SYMS_U64 raw_size = 0;
+          step_cursor += syms_based_range_read(base, range, step_cursor, 1, &raw_size);
+          SYMS_U64 size = SYMS_ClampTop(raw_size, 8);
           SYMS_U64 addr = syms_dw_expr__stack_pop(&stack);
           SYMS_B32 read_success = syms_false;
           SYMS_MemoryView *memview = config->memview;
diff --git a/Engine/Source/Runtime/SymsLib/syms/core/dwarf/syms_dwarf_expr.h b/Engine/Source/Runtime/SymsLib/syms/core/dwarf/syms_dwarf_expr.h
--- a/Engine/Source/Runtime/SymsLib/syms/core/dwarf/syms_dwarf_expr.h
+++ b/Engine/Source/Runtime/SymsLib/syms/core/dwarf/syms_dwarf_expr.h
@@ -37,6 +37,38 @@ typedef struct SYMS_DwRegsX64{
   };
 } SYMS_DwRegsX64;
 
+SYMS_GLOBAL SYMS_READ_ONLY SYMS_RegID syms_dw_reg_table_x64[] = {
+  SYMS_RegX64Code_rax,
+  SYMS_RegX64Code_rdx,
+  SYMS_RegX64Code_rcx,
+  SYMS_RegX64Code_rbx,
+  SYMS_RegX64Code_rsi,
+  SYMS_RegX64Code_rdi,
+  SYMS_RegX64Code_rbp,
+  SYMS_RegX64Code_rsp,
+  SYMS_RegX64Code_r8,
+  SYMS_RegX64Code_r9,
+  SYMS_RegX64Code_r10,
+  SYMS_RegX64Code_r11,
+  SYMS_RegX64Code_r12,
+  SYMS_RegX64Code_r13,
+  SYMS_RegX64Code_r14,
+  SYMS_RegX64Code_r15,
+  SYMS_RegX64Code_rip,
+};
+
+SYMS_GLOBAL SYMS_READ_ONLY SYMS_RegID syms_dw_reg_table_x86[] = {
+  SYMS_RegX86Code_eax,
+  SYMS_RegX86Code_edx,
+  SYMS_RegX86Code_ecx,
+  SYMS_RegX86Code_ebx,
+  SYMS_RegX86Code_esi,
+  SYMS_RegX86Code_edi,
+  SYMS_RegX86Code_ebp,
+  SYMS_RegX86Code_esp,
+  SYMS_RegX86Code_eip,
+};
+
 ////////////////////////////////
 //~ NOTE(allen): Dwarf Expression Eval Types
 
@@ -208,12 +240,12 @@ typedef struct SYMS_DwExprCallStack{
   SYMS_U64 depth;
 } SYMS_DwExprCallStack;
 
+
 ////////////////////////////////
 //~ NOTE(allen): Dwarf Expression Analysis & Eval Functions
 
 //- analyzers
 
-
 // NOTE(allen): This analyzer provides the most simplified dwarf expression
 // decoding. If the expression consists of a single op that can be interpreted
 // as a valid dwarf expression, then it represents that expression as a simple
diff --git a/Engine/Source/Runtime/SymsLib/syms/core/dwarf/syms_dwarf_parser.c b/Engine/Source/Runtime/SymsLib/syms/core/dwarf/syms_dwarf_parser.c
--- a/Engine/Source/Runtime/SymsLib/syms/core/dwarf/syms_dwarf_parser.c
+++ b/Engine/Source/Runtime/SymsLib/syms/core/dwarf/syms_dwarf_parser.c
@@ -1,4 +1,5 @@
 // Copyright Epic Games, Inc. All Rights Reserved.
+
 #ifndef SYMS_DWARF_PARSER_C
 #define SYMS_DWARF_PARSER_C
 
@@ -62,7 +63,7 @@
 ////////////////////////////////
 //~ rjf: Globals
 
-SYMS_GLOBAL SYMS_DwCompRoot syms_dw_nil_comp_root = {SYMS_DwCompUnitKind_RESERVED};
+SYMS_READ_ONLY SYMS_GLOBAL SYMS_DwCompRoot syms_dw_nil_comp_root = {SYMS_DwCompUnitKind_RESERVED};
 
 ////////////////////////////////
 //~ rjf: Basic Helpers
@@ -621,8 +622,8 @@ syms_dw_based_range_read_attrib_form_value(void *base, SYMS_U64Range range, SYMS
     {
       SYMS_U64 value1 = 0;
       SYMS_U64 value2 = 0;
-      bytes_read += syms_based_range_read_struct(base, range, offset, &value1);
-      bytes_read += syms_based_range_read_struct(base, range, offset, &value2);
+      bytes_read += syms_based_range_read_struct(base, range, offset,                    &value1);
+      bytes_read += syms_based_range_read_struct(base, range, offset + sizeof(SYMS_U64), &value2);
       form_value.v[0] = value1;
       form_value.v[1] = value2;
     }break;
@@ -754,6 +755,18 @@ syms_dw_dbg_accel_from_sec_info_array(SYMS_Arena *arena, SYMS_String8 data, SYMS
     has_necessary_sections = (match_count == SYMS_ARRAY_SIZE(necessary_kinds));
   }
   
+  //- rjf: scan sections for .text section, so we can store its index (useful for
+  // later stuff in the expression transpiler)
+  SYMS_U64 text_section_idx = 0;
+  for(SYMS_U64 idx = 0; idx < sections.count; idx += 1)
+  {
+    if(syms_string_match(sections.sec_info[idx].name, syms_str8_lit(".text"), 0))
+    {
+      text_section_idx = idx;
+      break;
+    }
+  }
+  
   //- rjf: build section map + figure out mode
   SYMS_DwMode mode = SYMS_DwMode_32Bit;
   SYMS_DwSection *section_map = 0;
@@ -950,6 +963,7 @@ syms_dw_dbg_accel_from_sec_info_array(SYMS_Arena *arena, SYMS_String8 data, SYMS
     dbg->mode = mode;
     dbg->vbase = vbase;
     dbg->sections = syms_dw_copy_sec_info_array(arena, sections);
+    dbg->text_section_idx = text_section_idx;
     dbg->acceptable_vrange = acceptable_vrange;
     dbg->section_map = section_map;
     dbg->unit_count = unit_count;
@@ -1189,7 +1203,7 @@ syms_dw_abbrev_offset_from_abbrev_id(SYMS_DwAbbrevTable table, SYMS_U64 abbrev_i
 //- rjf: .debug_ranges (DWARF V4)
 
 SYMS_API SYMS_U64RangeList
-syms_dw_v4_range_list_from_range_offset(SYMS_Arena *arena, SYMS_String8 data, SYMS_DwDbgAccel *dbg, SYMS_DwVersion version, SYMS_U64 addr_size, SYMS_U64 comp_unit_base_addr, SYMS_U64 range_off)
+syms_dw_v4_range_list_from_range_offset(SYMS_Arena *arena, SYMS_String8 data, SYMS_DwDbgAccel *dbg, SYMS_U64 addr_size, SYMS_U64 comp_unit_base_addr, SYMS_U64 range_off)
 {
   void *base = syms_dw_sec_base_from_dbg(data, dbg, SYMS_DwSectionKind_Ranges);
   SYMS_U64Range rng = syms_dw_sec_range_from_dbg(dbg, SYMS_DwSectionKind_Ranges);
@@ -1200,7 +1214,7 @@ syms_dw_v4_range_list_from_range_offset(SYMS_Arena *arena, SYMS_String8 data, SY
   SYMS_U64 read_off = range_off;
   SYMS_U64 base_addr = comp_unit_base_addr;
   
-  for(;;)
+  for(;read_off < rng.max;)
   {
     SYMS_U64 v0 = 0;
     SYMS_U64 v1 = 0;
@@ -1230,6 +1244,80 @@ syms_dw_v4_range_list_from_range_offset(SYMS_Arena *arena, SYMS_String8 data, SY
   return list;
 }
 
+//- rjf: .debug_loc (DWARF V4)
+
+SYMS_API SYMS_LocRangeList
+syms_dw_v4_location_ranges_from_loc_offset(SYMS_Arena *arena, SYMS_String8 data, SYMS_DwDbgAccel *dbg, SYMS_U64 addr_size, SYMS_U64 comp_unit_base_addr, SYMS_U64 offset)
+{
+  void *base = syms_dw_sec_base_from_dbg(data, dbg, SYMS_DwSectionKind_Loc);
+  SYMS_U64Range rng = syms_dw_sec_range_from_dbg(dbg, SYMS_DwSectionKind_Loc);
+  
+  SYMS_LocRangeList list = {0};
+  
+  SYMS_U64 read_off = offset;
+  SYMS_U64 base_addr = comp_unit_base_addr;
+  
+  for(;read_off < rng.max;)
+  {
+    SYMS_U64 v0 = 0;
+    SYMS_U64 v1 = 0;
+    read_off += syms_based_range_read(base, rng, read_off, addr_size, &v0);
+    read_off += syms_based_range_read(base, rng, read_off, addr_size, &v1);
+    
+    //- rjf: base address entry
+    if((addr_size == 4 && v0 == 0xffffffff) ||
+       (addr_size == 8 && v0 == 0xffffffffffffffff))
+    {
+      base_addr = v1;
+    }
+    //- rjf: end-of-list entry
+    else if(v0 == 0 && v1 == 0)
+    {
+      break;
+    }
+    //- rjf: location list entry
+    else
+    {
+      SYMS_U64 start_addr = v0 + base_addr;
+      SYMS_U64 end_addr = v1 + base_addr;
+      
+      // rjf: push to list
+      SYMS_LocRangeNode *n = syms_push_array_zero(arena, SYMS_LocRangeNode, 1);
+      n->loc_range.vrange = syms_make_u64_range(start_addr, end_addr);
+      n->loc_range.loc_id = (SYMS_LocID)read_off;
+      SYMS_QueuePush(list.first, list.last, n);
+      list.count += 1;
+      
+      // rjf: skip past location description
+      SYMS_U16 locdesc_length = 0;
+      read_off += syms_based_range_read_struct(base, rng, read_off, &locdesc_length);
+      read_off += locdesc_length;
+    }
+  }
+  
+  return list;
+}
+
+SYMS_API SYMS_Location
+syms_dw_v4_location_from_loc_id(SYMS_Arena *arena, SYMS_String8 data, SYMS_DwDbgAccel *dbg, SYMS_LocID id)
+{
+  SYMS_Location result = {0};
+  
+  void *base = syms_dw_sec_base_from_dbg(data, dbg, SYMS_DwSectionKind_Loc);
+  SYMS_U64Range sec_range = syms_dw_sec_range_from_dbg(dbg, SYMS_DwSectionKind_Loc);
+  SYMS_U64 read_off = (SYMS_U64)id;
+  SYMS_U16 length = 0;
+  read_off += syms_based_range_read_struct(base, sec_range, read_off, &length);
+  void *expr_base = syms_based_range_ptr(base, sec_range, read_off);
+  SYMS_U64Range expr_range = syms_make_u64_range(read_off, read_off+length);
+  SYMS_String8 location = syms_dw_expr__transpile_to_eval(arena, dbg, base, expr_range);
+  
+  // TODO(rjf): wire up `location` to `result`, once SYMS_Location changes to
+  // returning a SYMS_String8 instead of an op list.
+  
+  return result;
+}
+
 //- rjf: .debug_pubtypes + .debug_pubnames (DWARF V4)
 
 SYMS_API SYMS_DwPubStringsTable
@@ -1362,15 +1450,18 @@ syms_dw_v5_offset_from_offs_section_base_index(SYMS_String8 data, SYMS_DwDbgAcce
   SYMS_ASSERT_PARANOID(padding == 0); // must be 0 as of V5.
   
   //- rjf: read
-  void *entries = (SYMS_U8 *)sec_base + read_off;
-  SYMS_U64 count = (unit_length - sizeof(SYMS_U16)*2) / entry_len;
-  if(0 <= index && index < count)
+  if (unit_length >= sizeof(SYMS_U16)*2) 
   {
-    switch(entry_len)
+    void *entries = (SYMS_U8 *)sec_base + read_off;
+    SYMS_U64 count = (unit_length - sizeof(SYMS_U16)*2) / entry_len;
+    if(0 <= index && index < count)
     {
-      default: break;
-      case 4: {result = ((SYMS_U32 *)entries)[index];}break;
-      case 8: {result = ((SYMS_U64 *)entries)[index];}break;
+      switch(entry_len)
+      {
+        default: break;
+        case 4: {result = ((SYMS_U32 *)entries)[index];}break;
+        case 8: {result = ((SYMS_U64 *)entries)[index];}break;
+      }
     }
   }
   
@@ -1422,17 +1513,20 @@ syms_dw_v5_addr_from_addrs_section_base_index(SYMS_String8 data, SYMS_DwDbgAccel
   return result;
 }
 
-//- rjf: .debug_rnglists parsing
+//- rjf: .debug_rnglists + .debug_loclists parsing
 
 SYMS_API SYMS_U64
-syms_dw_v5_rnglist_offset_from_rnglist_section_base_index(SYMS_String8 data, SYMS_DwDbgAccel *dbg, SYMS_DwSectionKind section_kind, SYMS_U64 base, SYMS_U64 index)
+syms_dw_v5_sec_offset_from_rnglist_or_loclist_section_base_index(SYMS_String8 data, SYMS_DwDbgAccel *dbg, SYMS_DwSectionKind section_kind, SYMS_U64 base, SYMS_U64 index)
 {
   //
   // NOTE(rjf): This is only appropriate to call when SYMS_DwFormKind_RNGLISTX is
-  // used to access a range list. Otherwise, SYMS_DwFormKind_SEC_OFFSET is required.
+  // used to access a range list, *OR* when SYMS_DwFormKind_LOCLISTX is used to
+  // access a location list. Otherwise, SYMS_DwFormKind_SEC_OFFSET is required.
   //
-  // See the DWARF V5 spec (February 13, 2017), page 242.
+  // See the DWARF V5 spec (February 13, 2017), page 242. (rnglists)
+  // See the DWARF V5 spec (February 13, 2017), page 215. (loclists)
   //
+  
   SYMS_U64 result = 0;
   void *sec_base = syms_dw_sec_base_from_dbg(data, dbg, section_kind);
   SYMS_U64Range rng = syms_dw_sec_range_from_dbg(dbg, section_kind);
@@ -1481,7 +1575,7 @@ syms_dw_v5_rnglist_offset_from_rnglist_section_base_index(SYMS_String8 data, SYM
 }
 
 SYMS_API SYMS_U64RangeList
-syms_dw_v5_range_list_from_rnglist_or_loclist_offset(SYMS_Arena *arena, SYMS_String8 data, SYMS_DwDbgAccel *dbg, SYMS_DwSectionKind section, SYMS_U64 addr_size, SYMS_U64 addr_section_base, SYMS_U64 offset)
+syms_dw_v5_range_list_from_rnglist_offset(SYMS_Arena *arena, SYMS_String8 data, SYMS_DwDbgAccel *dbg, SYMS_DwSectionKind section, SYMS_U64 addr_size, SYMS_U64 addr_section_base, SYMS_U64 offset)
 {
   SYMS_U64RangeList list = {0};
   
@@ -1567,7 +1661,7 @@ syms_dw_v5_range_list_from_rnglist_or_loclist_offset(SYMS_Arena *arena, SYMS_Str
         SYMS_U64 start = 0;
         SYMS_U64 length = 0;
         read_off += syms_based_range_read(base, rng, read_off, addr_size, &start);
-        read_off += syms_based_range_read(base, rng, read_off, addr_size, &length);
+        read_off += syms_based_range_read_uleb128(base, rng, read_off, &length);
         syms_u64_range_list_push(arena, &list, syms_make_u64_range(start, start+length));
       }break;
     }
@@ -1576,6 +1670,155 @@ syms_dw_v5_range_list_from_rnglist_or_loclist_offset(SYMS_Arena *arena, SYMS_Str
   return list;
 }
 
+//- rjf: .debug_loclists parsing
+
+SYMS_API SYMS_LocRangeList
+syms_dw_v5_location_ranges_from_loclist_offset(SYMS_Arena *arena, SYMS_String8 data, SYMS_DwDbgAccel *dbg, SYMS_DwSectionKind section, SYMS_U64 addr_size, SYMS_U64 addr_section_base, SYMS_U64 offset)
+{
+  SYMS_LocRangeList list = {0};
+  
+  SYMS_U64 read_off = offset;
+  void *base = syms_dw_sec_base_from_dbg(data, dbg, section);
+  SYMS_U64Range rng = syms_dw_sec_range_from_dbg(dbg, section);
+  
+  SYMS_U64 base_addr = 0;
+  
+  for(SYMS_B32 done = syms_false; !done;)
+  {
+    SYMS_U8 kind8 = 0;
+    read_off += syms_based_range_read_struct(base, rng, read_off, &kind8);
+    SYMS_DwLocListEntryKind kind = (SYMS_DwLocListEntryKind)kind8;
+    
+    SYMS_B32 do_counted_location_description = 0;
+    SYMS_U64 start_addr = 0;
+    SYMS_U64 end_addr = 0;
+    
+    switch(kind)
+    {
+      //- rjf: can be used in split and non-split units:
+      default:
+      case SYMS_DwLocListEntryKind_EndOfList:
+      {
+        done = syms_true;
+      }break;
+      
+      case SYMS_DwLocListEntryKind_BaseAddressX:
+      {
+        SYMS_U64 base_addr_idx = 0;
+        read_off += syms_based_range_read_uleb128(base, rng, read_off, &base_addr_idx);
+        base_addr = syms_dw_v5_addr_from_addrs_section_base_index(data, dbg, SYMS_DwSectionKind_Addr, addr_section_base, base_addr_idx);
+      }break;
+      
+      case SYMS_DwLocListEntryKind_StartXEndX:
+      {
+        SYMS_U64 start_addr_idx = 0;
+        SYMS_U64 end_addr_idx = 0;
+        read_off += syms_based_range_read_uleb128(base, rng, read_off, &start_addr_idx);
+        read_off += syms_based_range_read_uleb128(base, rng, read_off, &end_addr_idx);
+        start_addr = syms_dw_v5_addr_from_addrs_section_base_index(data, dbg, SYMS_DwSectionKind_Addr, addr_section_base, start_addr_idx);
+        end_addr = syms_dw_v5_addr_from_addrs_section_base_index(data, dbg, SYMS_DwSectionKind_Addr, addr_section_base, end_addr_idx);
+        do_counted_location_description = 1;
+      }break;
+      
+      case SYMS_DwLocListEntryKind_StartXLength:
+      {
+        SYMS_U64 start_addr_idx = 0;
+        SYMS_U64 length = 0;
+        read_off += syms_based_range_read_uleb128(base, rng, read_off, &start_addr_idx);
+        read_off += syms_based_range_read_uleb128(base, rng, read_off, &length);
+        start_addr = syms_dw_v5_addr_from_addrs_section_base_index(data, dbg, SYMS_DwSectionKind_Addr, addr_section_base, start_addr_idx);
+        end_addr = start_addr + length;
+        do_counted_location_description = 1;
+      }break;
+      
+      case SYMS_DwLocListEntryKind_OffsetPair:
+      {
+        SYMS_U64 start_offset = 0;
+        SYMS_U64 end_offset = 0;
+        read_off += syms_based_range_read_uleb128(base, rng, read_off, &start_offset);
+        read_off += syms_based_range_read_uleb128(base, rng, read_off, &end_offset);
+        start_addr = start_offset + base_addr;
+        end_addr = end_offset + base_addr;
+        do_counted_location_description = 1;
+      }break;
+      
+      case SYMS_DwLocListEntryKind_DefaultLocation:
+      {
+        do_counted_location_description = 1;
+      }break;
+      
+      //- rjf: non-split units only:
+      
+      case SYMS_DwLocListEntryKind_BaseAddress:
+      {
+        SYMS_U64 new_base_addr = 0;
+        read_off += syms_based_range_read(base, rng, read_off, addr_size, &new_base_addr);
+        base_addr = new_base_addr;
+      }break;
+      
+      case SYMS_DwLocListEntryKind_StartEnd:
+      {
+        SYMS_U64 start = 0;
+        SYMS_U64 end = 0;
+        read_off += syms_based_range_read(base, rng, read_off, addr_size, &start);
+        read_off += syms_based_range_read(base, rng, read_off, addr_size, &end);
+        start_addr = start;
+        end_addr = end;
+        do_counted_location_description = 1;
+      }break;
+      
+      case SYMS_DwLocListEntryKind_StartLength:
+      {
+        SYMS_U64 start = 0;
+        SYMS_U64 length = 0;
+        read_off += syms_based_range_read(base, rng, read_off, addr_size, &start);
+        read_off += syms_based_range_read_uleb128(base, rng, read_off, &length);
+        start_addr = start;
+        end_addr = start+length;
+        do_counted_location_description = 1;
+      }break;
+    }
+    
+    //- rjf: parse a counted location description + push to result, if valid
+    if(do_counted_location_description)
+    {
+      SYMS_LocRangeNode *n = syms_push_array_zero(arena, SYMS_LocRangeNode, 1);
+      n->loc_range.vrange.min = start_addr;
+      n->loc_range.vrange.max = end_addr;
+      n->loc_range.loc_id = read_off;
+      SYMS_QueuePush(list.first, list.last, n);
+      list.count += 1;
+      
+      // rjf: skip past the location description
+      SYMS_U64 locdesc_length = 0;
+      read_off += syms_based_range_read_uleb128(base, rng, read_off, &locdesc_length);
+      read_off += locdesc_length;
+    }
+  }
+  
+  return list;
+}
+
+SYMS_API SYMS_Location
+syms_dw_v5_location_from_loclist_id(SYMS_Arena *arena, SYMS_String8 data, SYMS_DwDbgAccel *dbg, SYMS_DwSectionKind section, SYMS_LocID id)
+{
+  SYMS_Location result = {0};
+  
+  void *base = syms_dw_sec_base_from_dbg(data, dbg, section);
+  SYMS_U64Range sec_range = syms_dw_sec_range_from_dbg(dbg, section);
+  SYMS_U64 read_off = (SYMS_U64)id;
+  SYMS_U64 length = 0;
+  read_off += syms_based_range_read_uleb128(base, sec_range, read_off, &length);
+  void *expr_base = syms_based_range_ptr(base, sec_range, read_off);
+  SYMS_U64Range expr_range = syms_make_u64_range(read_off, read_off+length);
+  SYMS_String8 location = syms_dw_expr__transpile_to_eval(arena, dbg, base, expr_range);
+  
+  // TODO(rjf): wire up `location` to `result`, once SYMS_Location changes to
+  // returning a SYMS_String8 instead of an op list.
+  
+  return result;
+}
+
 ////////////////////////////////
 //~ rjf: Attrib Value Parsing
 
@@ -1619,10 +1862,23 @@ syms_dw_attrib_value_from_form_value(SYMS_String8 data, SYMS_DwDbgAccel *dbg, SY
   }
   //- rjf: (DWARF V5 ONLY) lookup into the .debug_loclists section via an index
   else if(resolve_params.version >= SYMS_DwVersion_V5 &&
-          value_class == SYMS_DwAttribClass_LOCLISTPTR &&
+          value_class == SYMS_DwAttribClass_LOCLIST &&
           form_kind == SYMS_DwFormKind_LOCLISTX)
   {
-    // TODO(rjf)
+    SYMS_U64 loclist_index = form_value.v[0];
+    SYMS_U64 loclist_offset = syms_dw_v5_sec_offset_from_rnglist_or_loclist_section_base_index(data, dbg, SYMS_DwSectionKind_LocLists, resolve_params.debug_loclists_base, loclist_index);
+    value.section = SYMS_DwSectionKind_LocLists;
+    value.v[0] = loclist_offset;
+  }
+  //- rjf: (DWARF V5 ONLY) lookup into the .debug_loclists section via an offset
+  else if(resolve_params.version >= SYMS_DwVersion_V5 &&
+          (value_class == SYMS_DwAttribClass_LOCLIST ||
+           value_class == SYMS_DwAttribClass_LOCLISTPTR) &&
+          form_kind == SYMS_DwFormKind_SEC_OFFSET)
+  {
+    SYMS_U64 loclist_offset = form_value.v[0];
+    value.section = SYMS_DwSectionKind_LocLists;
+    value.v[0] = loclist_offset;
   }
   //- rjf: (DWARF V5 ONLY) lookup into the .debug_rnglists section via an index
   else if(resolve_params.version >= SYMS_DwVersion_V5 &&
@@ -1631,7 +1887,7 @@ syms_dw_attrib_value_from_form_value(SYMS_String8 data, SYMS_DwDbgAccel *dbg, SY
           form_kind == SYMS_DwFormKind_RNGLISTX)
   {
     SYMS_U64 rnglist_index = form_value.v[0];
-    SYMS_U64 rnglist_offset = syms_dw_v5_rnglist_offset_from_rnglist_section_base_index(data, dbg, SYMS_DwSectionKind_RngLists, resolve_params.debug_rnglists_base, rnglist_index);
+    SYMS_U64 rnglist_offset = syms_dw_v5_sec_offset_from_rnglist_or_loclist_section_base_index(data, dbg, SYMS_DwSectionKind_RngLists, resolve_params.debug_rnglists_base, rnglist_index);
     value.section = SYMS_DwSectionKind_RngLists;
     value.v[0] = rnglist_offset;
   }
@@ -1733,6 +1989,16 @@ syms_dw_attrib_value_from_form_value(SYMS_String8 data, SYMS_DwDbgAccel *dbg, SY
     value.section = SYMS_DwSectionKind_Ranges;
     value.v[0] = ranges_offset;
   }
+  //- rjf: .debug_loc
+  else if(resolve_params.version < SYMS_DwVersion_V5 &&
+          (value_class == SYMS_DwAttribClass_LOCLISTPTR ||
+           value_class == SYMS_DwAttribClass_LOCLIST) &&
+          (form_kind == SYMS_DwFormKind_SEC_OFFSET))
+  {
+    SYMS_U64 offset = form_value.v[0];
+    value.section = SYMS_DwSectionKind_Loc;
+    value.v[0] = offset;
+  }
   //- rjf: invalid attribute class
   else if(value_class == 0)
   {
@@ -1770,13 +2036,13 @@ syms_dw_range_list_from_high_low_pc_and_ranges_attrib_value(SYMS_Arena *arena, S
     //- rjf: (DWARF V5 ONLY) .debug_rnglists offset
     case SYMS_DwSectionKind_RngLists:
     {
-      list = syms_dw_v5_range_list_from_rnglist_or_loclist_offset(arena, data, dbg, ranges_value.section, address_size, addr_section_base, ranges_value.v[0]);
+      list = syms_dw_v5_range_list_from_rnglist_offset(arena, data, dbg, ranges_value.section, address_size, addr_section_base, ranges_value.v[0]);
     }break;
     
     //- rjf: (DWARF V4 and earlier) .debug_ranges parsing
     case SYMS_DwSectionKind_Ranges:
     {
-      list = syms_dw_v4_range_list_from_range_offset(arena, data, dbg, SYMS_DwVersion_V4, address_size, comp_unit_base_addr, ranges_value.v[0]);
+      list = syms_dw_v4_range_list_from_range_offset(arena, data, dbg, address_size, comp_unit_base_addr, ranges_value.v[0]);
     }break;
     
     //- rjf: fall back to trying to use low/high PCs
@@ -1808,6 +2074,8 @@ syms_dw_parse_attrib_list_from_info_abbrev_offsets(SYMS_Arena *arena, SYMS_Strin
   SYMS_DwAttribListParseResult result = {0};
   for(SYMS_B32 good_abbrev = syms_true; good_abbrev;)
   {
+    SYMS_U64 attrib_info_offset = info_read_off;
+    
     //- rjf: parse abbrev attrib info
     SYMS_DwAbbrev abbrev;
     syms_memzero_struct(&abbrev);
@@ -1844,6 +2112,7 @@ syms_dw_parse_attrib_list_from_info_abbrev_offsets(SYMS_Arena *arena, SYMS_Strin
     if(good_abbrev)
     {
       SYMS_DwAttribNode *node = syms_push_array_zero(arena, SYMS_DwAttribNode, 1);
+      node->attrib.info_off     = attrib_info_offset;
       node->attrib.abbrev_id    = abbrev.id;
       node->attrib.attrib_kind  = attrib_kind;
       node->attrib.form_kind    = form_kind;
@@ -2240,6 +2509,7 @@ syms_dw_comp_root_from_range(SYMS_Arena *arena, SYMS_String8 data, SYMS_DwDbgAcc
   syms_memzero_struct(&unit);
   
   //- rjf: fill header data
+  unit.size            = size;
   unit.kind            = (SYMS_DwCompUnitKind)unit_kind;
   unit.version         = (SYMS_DwVersion)version;
   unit.address_size    = address_size;
@@ -4288,6 +4558,161 @@ syms_dw_children_from_sid_with_kinds(SYMS_Arena *arena, SYMS_String8 data, SYMS_
   return list;
 }
 
+//- rjf: variable locations
+
+SYMS_API SYMS_Location
+syms_dw_location_from_sid(SYMS_Arena *arena, SYMS_String8 data, SYMS_DwDbgAccel *dbg, SYMS_DwUnitAccel *unit, SYMS_SymbolID sid, SYMS_DwAttribKind loc_attrib)
+{
+  SYMS_Location result = {0};
+  SYMS_ArenaTemp scratch = syms_get_scratch(&arena, 1);
+  
+  SYMS_DwTagStub stub = syms_dw_cached_tag_stub_from_sid__parse_fallback(data, dbg, unit, sid);
+  if(stub.sid != 0)
+  {
+    SYMS_DwAttribList attribs = syms_dw_attrib_list_from_stub(scratch.arena, data, dbg, unit->address_size, &stub);
+    for(SYMS_DwAttribNode *n = attribs.first; n != 0; n = n->next)
+    {
+      SYMS_DwAttrib *attrib = &n->attrib;
+      if(attrib->attrib_kind == loc_attrib &&
+         attrib->form_kind == SYMS_DwFormKind_EXPRLOC)
+      {
+        SYMS_DwAttribValue location_value = syms_dw_attrib_value_from_form_value(data, dbg, unit->resolve_params,
+                                                                                 attrib->form_kind,
+                                                                                 attrib->value_class,
+                                                                                 attrib->form_value);
+        void *expr_base = syms_dw_sec_base_from_dbg(data, dbg, SYMS_DwSectionKind_Info);
+        SYMS_U64Range expr_range = syms_make_u64_range(location_value.v[0], location_value.v[0] + location_value.v[1]);
+        SYMS_String8 location = syms_dw_expr__transpile_to_eval(arena, dbg, expr_base, expr_range );
+        break;
+      }
+    }
+  }
+  
+  // TODO(rjf): once SYMS_Location has changed accordingly (returns a string
+  // instead of an op-list, wire up `location` to `result`)
+  
+  syms_release_scratch(scratch);
+  return result;
+}
+
+SYMS_API SYMS_LocRangeArray
+syms_dw_location_ranges_from_sid(SYMS_Arena *arena, SYMS_String8 data, SYMS_DwDbgAccel *dbg,
+                                 SYMS_DwUnitAccel *unit, SYMS_SymbolID sid, SYMS_DwAttribKind loc_attrib)
+{
+  SYMS_LocRangeArray result = {0};
+  SYMS_ArenaTemp scratch = syms_get_scratch(&arena, 1);
+  
+  // rjf: get location range list from sid, if applicable
+  SYMS_LocRangeList list = {0};
+  SYMS_DwTagStub stub = syms_dw_cached_tag_stub_from_sid__parse_fallback(data, dbg, unit, sid);
+  if(stub.sid != 0)
+  {
+    SYMS_DwAttribList attribs = syms_dw_attrib_list_from_stub(scratch.arena, data, dbg, unit->address_size, &stub);
+    for(SYMS_DwAttribNode *n = attribs.first; n != 0; n = n->next)
+    {
+      SYMS_DwAttrib *attrib = &n->attrib;
+      if(attrib->attrib_kind == loc_attrib &&
+         (attrib->value_class == SYMS_DwAttribClass_LOCLISTPTR ||
+          attrib->value_class == SYMS_DwAttribClass_LOCLIST))
+      {
+        SYMS_DwAttribValue location_value = syms_dw_attrib_value_from_form_value(data, dbg, unit->resolve_params,
+                                                                                 attrib->form_kind,
+                                                                                 attrib->value_class,
+                                                                                 attrib->form_value);
+        SYMS_U64 comp_unit_base_addr = unit->base_addr;
+        switch(location_value.section)
+        {
+          // NOTE(rjf): .debug_loclists is only available in DWARF V5, so we can use that to
+          // determine which parsing path to use.
+          case SYMS_DwSectionKind_LocLists:
+          {
+            list = syms_dw_v5_location_ranges_from_loclist_offset(scratch.arena, data, dbg, location_value.section, unit->address_size, comp_unit_base_addr, location_value.v[0]);
+          }break;
+          
+          // NOTE(rjf): .debug_loclists is only available in DWARF V4 and earlier.
+          case SYMS_DwSectionKind_Loc:
+          {
+            list = syms_dw_v4_location_ranges_from_loc_offset(scratch.arena, data, dbg, unit->address_size, comp_unit_base_addr, location_value.v[0]);
+          }break;
+        }
+        break;
+      }
+    }
+  }
+  
+  // rjf: convert list => array
+  if(list.count != 0)
+  {
+    result.count = list.count;
+    result.loc_ranges = syms_push_array_zero(arena, SYMS_LocRange, result.count);
+    SYMS_U64 idx = 0;
+    for(SYMS_LocRangeNode *n = list.first; n != 0; n = n->next, idx += 1)
+    {
+      result.loc_ranges[idx] = n->loc_range;
+    }
+  }
+  
+  syms_release_scratch(scratch);
+  return result;
+}
+
+SYMS_API SYMS_Location
+syms_dw_location_from_id(SYMS_Arena *arena, SYMS_String8 data, SYMS_DwDbgAccel *dbg,
+                         SYMS_DwUnitAccel *unit, SYMS_LocID loc_id)
+{
+  SYMS_Location result = {0};
+  switch(unit->version)
+  {
+    case SYMS_DwVersion_V5:{syms_dw_v5_location_from_loclist_id(arena, data, dbg, SYMS_DwSectionKind_LocLists, loc_id);}break;
+    case SYMS_DwVersion_V4:{syms_dw_v4_location_from_loc_id(arena, data, dbg, loc_id);}break;
+  }
+  return result;
+}
+
+//- rjf: location information helpers
+
+SYMS_API SYMS_DwAttribKind
+syms_dw_attrib_kind_from_proc_loc(SYMS_ProcLoc proc_loc)
+{
+  SYMS_DwAttribKind attrib = (SYMS_DwAttribKind)0;
+  switch(proc_loc)
+  {
+    case SYMS_ProcLoc_FrameBase:    {attrib = SYMS_DwAttribKind_FRAME_BASE;}break;
+    case SYMS_ProcLoc_ReturnAddress:{attrib = SYMS_DwAttribKind_RETURN_ADDR;}break;
+  }
+  return attrib;
+}
+
+SYMS_API SYMS_Location
+syms_dw_location_from_var_sid(SYMS_Arena *arena, SYMS_String8 data, SYMS_DwDbgAccel *dbg, SYMS_DwUnitAccel *unit, SYMS_SymbolID sid)
+{
+  return syms_dw_location_from_sid(arena, data, dbg, unit, sid, SYMS_DwAttribKind_LOCATION);
+}
+
+SYMS_API SYMS_LocRangeArray
+syms_dw_location_ranges_from_var_sid(SYMS_Arena *arena, SYMS_String8 data, SYMS_DwDbgAccel *dbg,
+                                     SYMS_DwUnitAccel *unit, SYMS_SymbolID sid)
+{
+  return syms_dw_location_ranges_from_sid(arena, data, dbg, unit, sid, SYMS_DwAttribKind_LOCATION);
+}
+
+SYMS_API SYMS_Location
+syms_dw_location_from_proc_sid(SYMS_Arena *arena, SYMS_String8 data, SYMS_DwDbgAccel *dbg, SYMS_DwUnitAccel *unit, SYMS_SymbolID sid, SYMS_ProcLoc proc_loc)
+{
+  SYMS_DwAttribKind attrib_kind = syms_dw_attrib_kind_from_proc_loc(proc_loc);
+  SYMS_Location location = syms_dw_location_from_sid(arena, data, dbg, unit, sid, attrib_kind);
+  return location;
+}
+
+SYMS_API SYMS_LocRangeArray
+syms_dw_location_ranges_from_proc_sid(SYMS_Arena *arena, SYMS_String8 data, SYMS_DwDbgAccel *dbg,
+                                      SYMS_DwUnitAccel *unit, SYMS_SymbolID sid, SYMS_ProcLoc proc_loc)
+{
+  SYMS_DwAttribKind attrib_kind = syms_dw_attrib_kind_from_proc_loc(proc_loc);
+  SYMS_LocRangeArray ranges = syms_dw_location_ranges_from_sid(arena, data, dbg, unit, sid, attrib_kind);
+  return ranges;
+}
+
 //- rjf: files
 
 SYMS_API SYMS_U64
@@ -4336,8 +4761,8 @@ syms_dw_file_name_from_id(SYMS_Arena *arena, SYMS_DwUnitSetAccel *unit_set, SYMS
 //- rjf: procedures
 
 SYMS_API SYMS_U64RangeArray
-syms_dw_proc_vranges_from_sid(SYMS_Arena *arena, SYMS_String8 data, SYMS_DwDbgAccel *dbg,
-                              SYMS_DwUnitAccel *unit, SYMS_SymbolID sid)
+syms_dw_scope_vranges_from_sid(SYMS_Arena *arena, SYMS_String8 data, SYMS_DwDbgAccel *dbg,
+                               SYMS_DwUnitAccel *unit, SYMS_SymbolID sid)
 {
   SYMS_U64RangeList list;
   syms_memzero_struct(&list);
@@ -4647,7 +5072,7 @@ syms_dw_parsed_line_table_from_comp_root(SYMS_Arena *arena, SYMS_String8 data, S
         // it works because compiler stores operand lengths that we can read
         // to skip unknown opcode
         else{
-          if (opcode <= vm_header.num_opcode_lens){
+          if (opcode > 0 && opcode <= vm_header.num_opcode_lens){
             SYMS_U8 num_operands = vm_header.opcode_lens[opcode - 1];
             for (SYMS_U8 i = 0; i < num_operands; i += 1){
               SYMS_U64 operand = 0;
@@ -4884,6 +5309,8 @@ syms_dw_line_parse_from_uid(SYMS_Arena *arena, SYMS_String8 data, SYMS_DwDbgAcce
 SYMS_API SYMS_U64
 syms_dw_read_line_file(void *line_base, SYMS_U64Range line_rng, SYMS_U64 line_off, SYMS_DwDbgAccel *dbg, SYMS_String8 data, SYMS_DwCompRoot *unit, SYMS_U8 address_size, SYMS_U64 format_count, SYMS_U64Range *formats, SYMS_DwLineFile *line_file_out)
 {
+  syms_memzero_struct(line_file_out);
+  
   SYMS_DwAttribValueResolveParams resolve_params = syms_dw_attrib_value_resolve_params_from_comp_root(unit);
   SYMS_U64 line_off_start = line_off;
   for (SYMS_U64 format_idx = 0; format_idx < format_count; ++format_idx) 
diff --git a/Engine/Source/Runtime/SymsLib/syms/core/dwarf/syms_dwarf_parser.h b/Engine/Source/Runtime/SymsLib/syms/core/dwarf/syms_dwarf_parser.h
--- a/Engine/Source/Runtime/SymsLib/syms/core/dwarf/syms_dwarf_parser.h
+++ b/Engine/Source/Runtime/SymsLib/syms/core/dwarf/syms_dwarf_parser.h
@@ -97,6 +97,7 @@ struct SYMS_DwDbgAccel
   SYMS_U64 vbase;
   SYMS_DwMode mode;
   SYMS_SecInfoArray sections;
+  SYMS_U64 text_section_idx;
   SYMS_U64Range acceptable_vrange;
   SYMS_DwSection *section_map; // NOTE(rjf): treat as [SYMS_DwSectionKind_COUNT] array
   SYMS_U64 unit_count;
@@ -195,6 +196,7 @@ struct SYMS_DwAttribValue
 typedef struct SYMS_DwAttrib SYMS_DwAttrib;
 struct SYMS_DwAttrib
 {
+  SYMS_U64 info_off;
   SYMS_U64 abbrev_id;
   SYMS_DwAttribKind attrib_kind;
   SYMS_DwFormKind form_kind;
@@ -239,6 +241,7 @@ typedef struct SYMS_DwCompRoot SYMS_DwCompRoot;
 struct SYMS_DwCompRoot
 {
   // NOTE(rjf): Header Data
+  SYMS_U64 size;
   SYMS_DwCompUnitKind kind;
   SYMS_DwVersion version;
   SYMS_U64 address_size;
@@ -602,10 +605,11 @@ SYMS_API SYMS_U64           syms_dw_abbrev_offset_from_abbrev_id(SYMS_DwAbbrevTa
 //~ rjf: Miscellaneous DWARF Section Parsing
 
 //- rjf: .debug_ranges (DWARF V4)
-SYMS_API SYMS_U64RangeList      syms_dw_v4_range_list_from_range_offset(SYMS_Arena *arena, SYMS_String8 data,
-                                                                        SYMS_DwDbgAccel *dbg,
-                                                                        SYMS_DwVersion version,
-                                                                        SYMS_U64 addr_size, SYMS_U64 comp_unit_base_addr, SYMS_U64 range_off);
+SYMS_API SYMS_U64RangeList      syms_dw_v4_range_list_from_range_offset(SYMS_Arena *arena, SYMS_String8 data, SYMS_DwDbgAccel *dbg, SYMS_U64 addr_size, SYMS_U64 comp_unit_base_addr, SYMS_U64 range_off);
+
+//- rjf: .debug_loc (DWARF V4)
+SYMS_API SYMS_LocRangeList      syms_dw_v4_location_ranges_from_loc_offset(SYMS_Arena *arena, SYMS_String8 data, SYMS_DwDbgAccel *dbg, SYMS_U64 addr_size, SYMS_U64 comp_unit_base_addr, SYMS_U64 offset);
+SYMS_API SYMS_Location          syms_dw_v4_location_from_loc_id(SYMS_Arena *arena, SYMS_String8 data, SYMS_DwDbgAccel *dbg, SYMS_LocID loc);
 
 //- rjf: .debug_pubtypes + .debug_pubnames (DWARF V4)
 SYMS_API SYMS_DwPubStringsTable syms_dw_v4_pub_strings_table_from_section_kind(SYMS_Arena *arena, SYMS_String8 data,
@@ -621,9 +625,13 @@ SYMS_API SYMS_U64 syms_dw_v5_offset_from_offs_section_base_index(SYMS_String8 da
 //- rjf: .debug_addr (DWARF V5)
 SYMS_API SYMS_U64 syms_dw_v5_addr_from_addrs_section_base_index(SYMS_String8 data, SYMS_DwDbgAccel *dbg, SYMS_DwSectionKind section, SYMS_U64 base, SYMS_U64 index);
 
-//- rjf: .debug_rnglists + .debug_loclists parsing
-SYMS_API SYMS_U64 syms_dw_v5_rnglist_offset_from_rnglist_section_base_index(SYMS_String8 data, SYMS_DwDbgAccel *dbg, SYMS_DwSectionKind section_kind, SYMS_U64 base, SYMS_U64 index);
-SYMS_API SYMS_U64RangeList syms_dw_v5_range_list_from_rnglist_or_loclist_offset(SYMS_Arena *arena, SYMS_String8 data, SYMS_DwDbgAccel *dbg, SYMS_DwSectionKind section, SYMS_U64 addr_size, SYMS_U64 addr_section_base, SYMS_U64 offset);
+//- rjf: .debug_rnglists parsing (DWARF V5)
+SYMS_API SYMS_U64 syms_dw_v5_sec_offset_from_rnglist_or_loclist_section_base_index(SYMS_String8 data, SYMS_DwDbgAccel *dbg, SYMS_DwSectionKind section_kind, SYMS_U64 base, SYMS_U64 index);
+SYMS_API SYMS_U64RangeList syms_dw_v5_range_list_from_rnglist_offset(SYMS_Arena *arena, SYMS_String8 data, SYMS_DwDbgAccel *dbg, SYMS_DwSectionKind section, SYMS_U64 addr_size, SYMS_U64 addr_section_base, SYMS_U64 offset);
+
+//- rjf: .debug_loclists parsing (DWARF V5)
+SYMS_API SYMS_LocRangeList syms_dw_v5_location_ranges_from_loclist_offset(SYMS_Arena *arena, SYMS_String8 data, SYMS_DwDbgAccel *dbg, SYMS_DwSectionKind section, SYMS_U64 addr_size, SYMS_U64 addr_section_base, SYMS_U64 offset);
+SYMS_API SYMS_Location syms_dw_v5_location_from_loclist_id(SYMS_Arena *arena, SYMS_String8 data, SYMS_DwDbgAccel *dbg, SYMS_DwSectionKind section, SYMS_LocID id);
 
 ////////////////////////////////
 //~ rjf: Attrib Value Parsing
@@ -750,6 +758,22 @@ SYMS_API SYMS_U64          syms_dw_voff_from_var_sid(SYMS_String8 data, SYMS_DwD
                                                      SYMS_DwUnitAccel *unit, SYMS_SymbolID sid);
 SYMS_API SYMS_SymbolIDList syms_dw_children_from_sid_with_kinds(SYMS_Arena *arena, SYMS_String8 data, SYMS_DwDbgAccel *dbg, SYMS_DwUnitAccel *unit, SYMS_SymbolID sid, SYMS_DwTagKind *kinds, SYMS_U64 count);
 
+//- rjf: location information
+SYMS_API SYMS_Location     syms_dw_location_from_sid(SYMS_Arena *arena, SYMS_String8 data, SYMS_DwDbgAccel *dbg, SYMS_DwUnitAccel *unit, SYMS_SymbolID sid, SYMS_DwAttribKind loc_attrib);
+SYMS_API SYMS_LocRangeArray syms_dw_location_ranges_from_sid(SYMS_Arena *arena, SYMS_String8 data, SYMS_DwDbgAccel *dbg,
+                                                             SYMS_DwUnitAccel *unit, SYMS_SymbolID sid, SYMS_DwAttribKind loc_attrib);
+SYMS_API SYMS_Location     syms_dw_location_from_id(SYMS_Arena *arena, SYMS_String8 data, SYMS_DwDbgAccel *dbg,
+                                                    SYMS_DwUnitAccel *unit, SYMS_LocID loc_id);
+
+//- rjf: location information helpers
+SYMS_API SYMS_DwAttribKind syms_dw_attrib_kind_from_proc_loc(SYMS_ProcLoc proc_loc);
+SYMS_API SYMS_Location     syms_dw_location_from_var_sid(SYMS_Arena *arena, SYMS_String8 data, SYMS_DwDbgAccel *dbg, SYMS_DwUnitAccel *unit, SYMS_SymbolID sid);
+SYMS_API SYMS_LocRangeArray syms_dw_location_ranges_from_var_sid(SYMS_Arena *arena, SYMS_String8 data, SYMS_DwDbgAccel *dbg,
+                                                                 SYMS_DwUnitAccel *unit, SYMS_SymbolID sid);
+SYMS_API SYMS_Location     syms_dw_location_from_proc_sid(SYMS_Arena *arena, SYMS_String8 data, SYMS_DwDbgAccel *dbg, SYMS_DwUnitAccel *unit, SYMS_SymbolID sid, SYMS_ProcLoc proc_loc);
+SYMS_API SYMS_LocRangeArray syms_dw_location_ranges_from_proc_sid(SYMS_Arena *arena, SYMS_String8 data, SYMS_DwDbgAccel *dbg,
+                                                                  SYMS_DwUnitAccel *unit, SYMS_SymbolID sid, SYMS_ProcLoc proc_loc);
+
 //- rjf: files
 SYMS_API SYMS_U64          syms_dw_file_index_from_id(SYMS_FileID file_id);
 SYMS_API SYMS_FileID       syms_dw_file_id_from_index(SYMS_U64 idx);
@@ -757,9 +781,9 @@ SYMS_API SYMS_String8      syms_dw_file_name_from_id(SYMS_Arena *arena, SYMS_DwU
                                                      SYMS_UnitID uid, SYMS_FileID file_id);
 
 //- rjf: procedures
-SYMS_API SYMS_U64RangeArray syms_dw_proc_vranges_from_sid(SYMS_Arena *arena, SYMS_String8 data,
-                                                          SYMS_DwDbgAccel *dbg, SYMS_DwUnitAccel *unit,
-                                                          SYMS_SymbolID sid);
+SYMS_API SYMS_U64RangeArray syms_dw_scope_vranges_from_sid(SYMS_Arena *arena, SYMS_String8 data,
+                                                           SYMS_DwDbgAccel *dbg, SYMS_DwUnitAccel *unit,
+                                                           SYMS_SymbolID sid);
 SYMS_API SYMS_SigInfo      syms_dw_sig_info_from_sid(SYMS_Arena *arena, SYMS_String8 data, SYMS_DwDbgAccel *dbg,
                                                      SYMS_DwUnitAccel *unit, SYMS_SymbolID sid);
 SYMS_API SYMS_UnitIDAndSig syms_dw_proc_sig_handle_from_sid(SYMS_String8 data, SYMS_DwDbgAccel *dbg,
diff --git a/Engine/Source/Runtime/SymsLib/syms/core/elf/syms_elf.h b/Engine/Source/Runtime/SymsLib/syms/core/elf/syms_elf.h
--- a/Engine/Source/Runtime/SymsLib/syms/core/elf/syms_elf.h
+++ b/Engine/Source/Runtime/SymsLib/syms/core/elf/syms_elf.h
@@ -273,6 +273,15 @@ typedef struct
   SYMS_U16 vs_vers;
 } SYMS_ElfExternalVersym;
 
+typedef struct
+{
+  SYMS_U32 name_size;
+  SYMS_U32 desc_size;
+  SYMS_U32 type;
+  // name + desc
+  // SYMS_U8  data[1];
+} SYMS_ElfNote;
+
 SYMS_C_LINKAGE_BEGIN
 
 ////////////////////////////////
diff --git a/Engine/Source/Runtime/SymsLib/syms/core/elf/syms_elf_parser.c b/Engine/Source/Runtime/SymsLib/syms/core/elf/syms_elf_parser.c
--- a/Engine/Source/Runtime/SymsLib/syms/core/elf/syms_elf_parser.c
+++ b/Engine/Source/Runtime/SymsLib/syms/core/elf/syms_elf_parser.c
@@ -1,4 +1,5 @@
 // Copyright Epic Games, Inc. All Rights Reserved.
+
 #ifndef SYMS_ELF_PARSER_C
 #define SYMS_ELF_PARSER_C
 
@@ -95,8 +96,7 @@ syms_elf_img_header_from_file(SYMS_String8 file)
     SYMS_U64 program_header_off = elf_header.e_phoff;
     SYMS_U64 program_header_num = (elf_header.e_phnum != SYMS_U16_MAX) ? elf_header.e_phnum : section_header.sh_info;
     
-    //- rjf: search for base address, which is stored in program header that is closest to entry point
-    SYMS_U64 prev_delta = SYMS_U64_MAX;
+    //- rjf: search for base address, by grabbing the first LOAD phdr
     for(SYMS_U64 i = 0; i < program_header_num; i += 1)
     {
       SYMS_ElfPhdr64 program_header;
@@ -110,24 +110,15 @@ syms_elf_img_header_from_file(SYMS_String8 file)
           syms_based_range_read_struct(file_base, file_range, program_header_off + i*sizeof(SYMS_ElfPhdr32), &h32);
           program_header = syms_elf_phdr64_from_phdr32(h32);
         }break;
-        
         case SYMS_ElfClass_64:
         {
           syms_based_range_read_struct(file_base, file_range, program_header_off + i*sizeof(SYMS_ElfPhdr64), &program_header);
         }break;
       }
-      
       if(program_header.p_type == SYMS_ElfPKind_Load)
       {
-        if(elf_header.e_entry >= program_header.p_vaddr)
-        {
-          SYMS_U64 delta = elf_header.e_entry - program_header.p_vaddr;
-          if(delta < prev_delta)
-          {
-            base_address = program_header.p_vaddr;
-            prev_delta = delta;
-          }
-        }
+        base_address = program_header.p_vaddr;
+        break;
       }
     }
   }
@@ -222,10 +213,9 @@ syms_elf_section_array_from_img_header(SYMS_Arena *arena, SYMS_String8 file, SYM
     }
     
     // rjf: fill section data
-    sections[section_idx].code = (SYMS_ElfSectionCode)header.sh_type;
+    sections[section_idx].header = header;
     sections[section_idx].virtual_range = syms_make_u64_range(header.sh_addr, header.sh_addr + virt_size);
     sections[section_idx].file_range = syms_make_u64_range(header.sh_offset, header.sh_offset + file_size);
-    sections[section_idx].sh_link = header.sh_link;
     sections[section_idx].name = name_stabilized;
   }
   
@@ -243,6 +233,35 @@ syms_elf_section_array_from_img_header(SYMS_Arena *arena, SYMS_String8 file, SYM
   return result;
 }
 
+SYMS_API SYMS_ElfSegmentArray
+syms_elf_segment_array_from_img_header(SYMS_Arena *arena, SYMS_String8 file, SYMS_ElfImgHeader img)
+{
+  void *base = file.str;
+  SYMS_U64Range range = syms_make_u64_range(0, file.size);
+
+  SYMS_U64 segment_count = img.ehdr.e_phnum;
+  SYMS_ElfPhdr64 *segments = syms_push_array_zero(arena, SYMS_ElfPhdr64, segment_count);
+  for(SYMS_U64 segment_idx = 0; segment_idx < segment_count; segment_idx += 1) 
+  {
+    if(img.is_32bit) 
+    {
+      SYMS_ElfPhdr32 phdr32;
+      syms_based_range_read_struct(base, range, img.ehdr.e_phoff + segment_idx * sizeof(SYMS_ElfPhdr32), &phdr32);
+      segments[segment_idx] = syms_elf_phdr64_from_phdr32(phdr32);
+    }
+    else
+    {
+      syms_based_range_read_struct(base, range, img.ehdr.e_phoff + segment_idx * sizeof(SYMS_ElfPhdr64), &segments[segment_idx]);
+    }
+  }
+
+  SYMS_ElfSegmentArray result;
+  result.count = segment_count;
+  result.v = segments;
+
+  return result;
+}
+
 SYMS_API SYMS_ElfExtDebugRef
 syms_elf_ext_debug_ref_from_elf_section_array(SYMS_String8 file, SYMS_ElfSectionArray sections)
 {
@@ -319,6 +338,7 @@ syms_elf_bin_accel_from_file(SYMS_Arena *arena, SYMS_String8 data, SYMS_ElfFileA
   syms_memmove(&bin_accel->header, &file_accel->header, sizeof(file_accel->header));
   bin_accel->format = file_accel->format;
   bin_accel->sections = syms_elf_section_array_from_img_header(arena, data, file_accel->header);
+  bin_accel->segments = syms_elf_segment_array_from_img_header(arena, data, file_accel->header);
   return bin_accel;
 }
 
@@ -409,8 +429,9 @@ syms_elf_imports_from_bin(SYMS_Arena *arena, SYMS_String8 data, SYMS_ElfBinAccel
     {
       // NOTE(rjf): We subtract 1 from sh_link because we don't store the null
       // section in the bin section array.
-      SYMS_U64 sh_link_idx = sect->sh_link-1;
-      SYMS_B32 sh_link_in_bounds = (0 < sect->sh_link && sect->sh_link < bin->sections.count);
+      SYMS_U64 sh_link_idx = sect->header.sh_link-1;
+      SYMS_B32 sh_link_in_bounds = (0 < sect->header.sh_link && sect->header.sh_link < bin->sections.count);
+      (void)sh_link_in_bounds;
       if(syms_string_match(sect->name, syms_str8_lit(".gnu.version_r"), 0)) 
       {
         verneed_range = sect->file_range;
@@ -429,10 +450,13 @@ syms_elf_imports_from_bin(SYMS_Arena *arena, SYMS_String8 data, SYMS_ElfBinAccel
   }
   
   //- rjf: grab .dynsym symbol array
-  SYMS_U64 sym_arr_count = syms_u64_range_size(dynsym_range) / sym_size;
-  SYMS_ElfSym64 *sym_arr = syms_push_array_zero(scratch.arena, SYMS_ElfSym64, sym_arr_count);
-  SYMS_ElfSym64 *sym_arr_opl = sym_arr + sym_arr_count;
+  SYMS_ElfSym64 *sym_arr = 0;
+  SYMS_ElfSym64 *sym_arr_opl = 0;
+  if (syms_u64_range_size(dynsym_range) > sym_size) 
   {
+    SYMS_U64 sym_arr_count = (syms_u64_range_size(dynsym_range) / sym_size) - 1;
+    sym_arr = syms_push_array_zero(scratch.arena, SYMS_ElfSym64, sym_arr_count);
+    sym_arr_opl = sym_arr + sym_arr_count;
     // NOTE(rjf): We start the read offset at sym_size because the format
     // always has a meaningless null symbol at the beginning of the block.
     SYMS_U64 roff__sym_array = sym_size;
@@ -580,9 +604,9 @@ syms_elf_exports_from_bin(SYMS_Arena *arena, SYMS_String8 data, SYMS_ElfBinAccel
     SYMS_ElfSection *sect = &bin->sections.v[sec_idx];
     if (syms_string_match(sect->name, syms_str8_lit(".gnu.version_d"), 0)) 
     {
-      SYMS_ASSERT(sect->sh_link > 0 && sect->sh_link < bin->sections.count);
+      SYMS_ASSERT(sect->header.sh_link > 0 && sect->header.sh_link < bin->sections.count);
       verdef_range = sect->file_range;
-      verdef_strtab_range = bin->sections.v[sect->sh_link-1].file_range;
+      verdef_strtab_range = bin->sections.v[sect->header.sh_link-1].file_range;
     }
     if (syms_string_match(sect->name, syms_str8_lit(".gnu.version"), 0)) 
     {
@@ -590,9 +614,9 @@ syms_elf_exports_from_bin(SYMS_Arena *arena, SYMS_String8 data, SYMS_ElfBinAccel
     }
     if (syms_string_match(sect->name, syms_str8_lit(".dynsym"), 0)) 
     {
-      SYMS_ASSERT(sect->sh_link > 0 && sect->sh_link < bin->sections.count);
+      SYMS_ASSERT(sect->header.sh_link > 0 && sect->header.sh_link < bin->sections.count);
       dynsym_range = sect->file_range;
-      dynsym_strtab_range = bin->sections.v[sect->sh_link-1].file_range;
+      dynsym_strtab_range = bin->sections.v[sect->header.sh_link-1].file_range;
     }
   }
   
diff --git a/Engine/Source/Runtime/SymsLib/syms/core/elf/syms_elf_parser.h b/Engine/Source/Runtime/SymsLib/syms/core/elf/syms_elf_parser.h
--- a/Engine/Source/Runtime/SymsLib/syms/core/elf/syms_elf_parser.h
+++ b/Engine/Source/Runtime/SymsLib/syms/core/elf/syms_elf_parser.h
@@ -22,11 +22,10 @@ struct SYMS_ElfImgHeader
 typedef struct SYMS_ElfSection SYMS_ElfSection;
 struct SYMS_ElfSection
 {
-  SYMS_ElfSectionCode code;
+  SYMS_ElfShdr64 header;
   SYMS_U64Range virtual_range;
   SYMS_U64Range file_range;
   SYMS_String8 name;
-  SYMS_U32 sh_link;
 };
 
 typedef struct SYMS_ElfSectionArray SYMS_ElfSectionArray;
@@ -36,6 +35,13 @@ struct SYMS_ElfSectionArray
   SYMS_ElfSection *v;
 };
 
+typedef struct SYMS_ElfSegmentArray SYMS_ElfSegmentArray;
+struct SYMS_ElfSegmentArray
+{
+  SYMS_U64 count;
+  SYMS_ElfPhdr64 *v;
+};
+
 typedef struct SYMS_ElfExtDebugRef SYMS_ElfExtDebugRef;
 struct SYMS_ElfExtDebugRef
 {
@@ -62,6 +68,7 @@ struct SYMS_ElfBinAccel
   SYMS_FileFormat format;
   SYMS_ElfImgHeader header;
   SYMS_ElfSectionArray sections;
+  SYMS_ElfSegmentArray segments;
 };
 
 SYMS_C_LINKAGE_BEGIN
@@ -69,9 +76,10 @@ SYMS_C_LINKAGE_BEGIN
 ////////////////////////////////
 //~ rjf: Low-Level Header/Section Parsing
 
-SYMS_API SYMS_ElfImgHeader syms_elf_img_header_from_file(SYMS_String8 file);
+SYMS_API SYMS_ElfImgHeader    syms_elf_img_header_from_file(SYMS_String8 file);
+SYMS_API SYMS_ElfExtDebugRef  syms_elf_ext_debug_ref_from_elf_section_array(SYMS_String8 file, SYMS_ElfSectionArray sections);
 SYMS_API SYMS_ElfSectionArray syms_elf_section_array_from_img_header(SYMS_Arena *arena, SYMS_String8 file, SYMS_ElfImgHeader img);
-SYMS_API SYMS_ElfExtDebugRef syms_elf_ext_debug_ref_from_elf_section_array(SYMS_String8 file, SYMS_ElfSectionArray sections);
+SYMS_API SYMS_ElfSegmentArray syms_elf_segment_array_from_img_header(SYMS_Arena *arena, SYMS_String8 file, SYMS_ElfImgHeader img);
 
 ////////////////////////////////
 //~ rjf: High-Level API Canonical Conversions
