- Updated Sysmlib that can handle > 4gb pdb files

diff --git a/Engine/Source/Runtime/SymsLib/syms/core/group/syms_group.c b/Engine/Source/Runtime/SymsLib/syms/core/group/syms_group.c
--- a/Engine/Source/Runtime/SymsLib/syms/core/group/syms_group.c
+++ b/Engine/Source/Runtime/SymsLib/syms/core/group/syms_group.c
@@ -1,10 +1,8 @@
 // Copyright Epic Games, Inc. All Rights Reserved.
+
 #ifndef SYMS_GROUP_C
 #define SYMS_GROUP_C
 
-// TODO(allen): killing the arena copy optimizations may cause some problems
-// in group, take a closer look into usage of string cons again.
-
 ////////////////////////////////
 // NOTE(allen): Syms Group Setup Functions
 
@@ -23,8 +21,8 @@ syms_group_release(SYMS_Group *group){
 
 SYMS_API void
 syms_group_init(SYMS_Group *group, SYMS_ParseBundle *params){
-  SYMS_ProfBegin("syms_group_init");
   SYMS_ASSERT_PARANOID(syms_thread_lane == 0);
+  SYMS_ProfBegin("syms_group_init");
   
   //- fill in bin
   SYMS_String8 bin_data = params->bin_data;
@@ -105,6 +103,7 @@ syms_group_init(SYMS_Group *group, SYMS_ParseBundle *params){
   group->units = syms_push_array_zero(group->arena, SYMS_UnitAccel*, unit_count);
   group->proc_sid_arrays = syms_push_array_zero(group->arena, SYMS_SymbolIDArray, unit_count);
   group->var_sid_arrays = syms_push_array_zero(group->arena, SYMS_SymbolIDArray, unit_count);
+  group->thread_sid_arrays = syms_push_array_zero(group->arena, SYMS_SymbolIDArray, unit_count);
   group->type_sid_arrays = syms_push_array_zero(group->arena, SYMS_SymbolIDArray, unit_count);
   group->file_tables = syms_push_array_zero(group->arena, SYMS_String8Array, unit_count);
   group->inferred_file_tables = syms_push_array_zero(group->arena, SYMS_String8Array, unit_count);
@@ -114,16 +113,16 @@ syms_group_init(SYMS_Group *group, SYMS_ParseBundle *params){
   group->line_sequence_maps = syms_push_array_zero(group->arena, SYMS_SpatialMap1D, unit_count);
   group->file_to_line_to_addr_buckets = syms_push_array_zero(group->arena, SYMS_FileToLineToAddrBuckets,
                                                              unit_count);
+  group->unit_type_maps = syms_push_array_zero(group->arena, SYMS_SymbolNameMap, unit_count);
   
   //- setup hash table caches
   group->string_cons = syms_string_cons_alloc(group->arena, 4093);
   group->file_id_2_name_map = syms_file_id_2_name_map_alloc(group->arena, 4093);
   
-  group->type_usid_buckets.bucket_count = 4093;
-  group->type_usid_buckets.buckets = syms_push_array_zero(group->arena, SYMS_TypeUSIDNode*, 4093);
-  
-  group->type_content_buckets.bucket_count = 4093;
-  group->type_content_buckets.buckets = syms_push_array_zero(group->arena, SYMS_TypeContentNode*, 4093);
+  //- setup type graph
+  syms_type_graph_init(&group->type_graph,
+                       group->arena, &group->string_cons,
+                       group->address_size);
   
   SYMS_ProfEnd();
 }
@@ -330,6 +329,24 @@ syms_group_var_sid_array_from_uid(SYMS_Group *group, SYMS_UnitID uid){
   return(result);
 }
 
+SYMS_API SYMS_SymbolIDArray*
+syms_group_tls_var_sid_array_from_uid(SYMS_Group *group, SYMS_UnitID uid){
+  SYMS_ProfBegin("syms_group_tls_var_sid_array_from_uid");
+  SYMS_SymbolIDArray *result = &syms_sid_array_nil;
+  if (1 <= uid && uid < group->unit_count){
+    SYMS_U64 index = uid - 1;
+    if ((group->unit_cache_flags[index] & SYMS_GroupUnitCacheFlag_HasTlsVarSidArray) == 0){
+      SYMS_UnitAccel *unit = syms_group_unit_from_uid(group, uid);
+      SYMS_Arena *arena = syms_group_get_lane_arena(group);
+      group->thread_sid_arrays[index] = syms_tls_var_sid_array_from_unit(arena, unit);
+      group->unit_cache_flags[index] |= SYMS_GroupUnitCacheFlag_HasTlsVarSidArray;
+    }
+    result = &group->thread_sid_arrays[index];
+  }
+  SYMS_ProfEnd();
+  return(result);
+}
+
 SYMS_API SYMS_SymbolIDArray*
 syms_group_type_sid_array_from_uid(SYMS_Group *group, SYMS_UnitID uid){
   SYMS_ProfBegin("syms_group_type_sid_array_from_uid");
@@ -489,8 +506,8 @@ syms_group_default_vbase(SYMS_Group *group){
 
 SYMS_API SYMS_String8
 syms_group_file_name_from_id(SYMS_Group *group, SYMS_UnitID uid, SYMS_FileID file_id){
-  SYMS_ProfBegin("syms_group_file_name_from_id");
   SYMS_ASSERT_PARANOID(syms_thread_lane == 0);
+  SYMS_ProfBegin("syms_group_file_name_from_id");
   SYMS_String8 from_map = syms_file_id_2_name_map_name_from_id(&group->file_id_2_name_map, uid, file_id);
   SYMS_String8 result = {0};
   if (from_map.size != 0){
@@ -635,7 +652,7 @@ syms_group_proc_sid_from_uid_voff__linear_scan(SYMS_Group *group, SYMS_UnitID ui
     SYMS_SymbolKind kind = syms_group_symbol_kind_from_sid(group, unit, sid);
     SYMS_B32 hit = syms_false;
     if (kind == SYMS_SymbolKind_Procedure){
-      SYMS_U64RangeArray ranges = syms_proc_vranges_from_sid(scratch.arena, dbg_data, dbg, unit, sid);
+      SYMS_U64RangeArray ranges = syms_scope_vranges_from_sid(scratch.arena, dbg_data, dbg, unit, sid);
       SYMS_U64Range *range = ranges.ranges;
       for (SYMS_U64 i = 0; i < ranges.count; i += 1, range += 1){
         if (range->min <= voff && voff < range->max){
@@ -656,45 +673,6 @@ syms_group_proc_sid_from_uid_voff__linear_scan(SYMS_Group *group, SYMS_UnitID ui
   return(result);
 }
 
-SYMS_API SYMS_SymbolID
-syms_group_var_sid_from_uid_voff__linear_scan(SYMS_Group *group, SYMS_UnitID uid, SYMS_U64 voff){
-  SYMS_ProfBegin("syms_group_var_sid_from_uid_voff__linear_scan");
-  
-  SYMS_String8 dbg_data = group->dbg_data;
-  SYMS_DbgAccel *dbg = group->dbg;
-  SYMS_UnitAccel *unit = syms_group_unit_from_uid(group, uid);
-  SYMS_SymbolIDArray *sid_array = syms_group_var_sid_array_from_uid(group, uid);
-  SYMS_U64 count = sid_array->count;
-  
-  SYMS_SymbolID result = {0};
-  SYMS_SymbolID *sid_ptr = sid_array->ids;
-  for (SYMS_U64 i = 0; i < count; i += 1, sid_ptr += 1){
-    SYMS_SymbolID sid = *sid_ptr;
-    SYMS_SymbolKind kind = syms_group_symbol_kind_from_sid(group, unit, sid);
-    SYMS_B32 hit = syms_false;
-    if (kind == SYMS_SymbolKind_ImageRelativeVariable){
-      SYMS_U64 var_voff = syms_voff_from_var_sid(dbg_data, dbg, unit, sid);
-      if (var_voff == voff){
-        hit = syms_true;
-      }
-      else if (var_voff < voff){
-        SYMS_USID var_type_usid = syms_type_from_var_sid(dbg_data, dbg, unit, sid);
-        SYMS_U64 var_size = syms_group_type_size_from_usid(group, var_type_usid);
-        if (voff < var_voff + var_size){
-          hit = syms_true;
-        }
-      }
-    }
-    if (hit){
-      result = sid;
-      break;
-    }
-  }
-  
-  SYMS_ProfEnd();
-  return(result);
-}
-
 SYMS_API SYMS_Line
 syms_group_line_from_uid_voff__linear_scan(SYMS_Group *group, SYMS_UnitID uid, SYMS_U64 voff){
   SYMS_ProfBegin("syms_group_line_from_uid_voff__linear_scan");
@@ -905,7 +883,7 @@ syms_group_proc_map_from_uid(SYMS_Group *group, SYMS_UnitID uid){
       SYMS_SymbolID sid = *sid_ptr;
       SYMS_SymbolKind kind = syms_group_symbol_kind_from_sid(group, unit, sid);
       if (kind == SYMS_SymbolKind_Procedure){
-        SYMS_U64RangeArray ranges = syms_proc_vranges_from_sid(scratch.arena, dbg_data, dbg, unit, sid);
+        SYMS_U64RangeArray ranges = syms_scope_vranges_from_sid(scratch.arena, dbg_data, dbg, unit, sid);
         syms_spatial_map_1d_loose_push(scratch.arena, &loose, sid, ranges);
       }
     }
@@ -987,58 +965,6 @@ syms_group_line_sequence_map_from_uid(SYMS_Group *group, SYMS_UnitID uid){
   return(result);
 }
 
-SYMS_API SYMS_SpatialMap1D*
-syms_group_var_map_from_uid(SYMS_Group *group, SYMS_UnitID uid){
-  SYMS_ProfBegin("syms_group_var_map_from_uid");
-  SYMS_ASSERT_PARANOID(syms_thread_lane == 0);
-  SYMS_U64 index = uid - 1;
-  if (index < group->unit_count &&
-      (group->unit_cache_flags[index] & SYMS_GroupUnitCacheFlag_HasVarMap) == 0){
-    //- build loose map
-    SYMS_ArenaTemp scratch = syms_get_scratch(0, 0);
-    SYMS_SpatialMap1DLoose loose = {0};
-    
-    SYMS_String8 dbg_data = group->dbg_data;
-    SYMS_DbgAccel *dbg = group->dbg;
-    SYMS_UnitAccel *unit = syms_group_unit_from_uid(group, uid);
-    SYMS_SymbolIDArray *sid_array = syms_group_var_sid_array_from_uid(group, uid);
-    SYMS_U64 count = sid_array->count;
-    
-    SYMS_SymbolID *sid_ptr = sid_array->ids;
-    for (SYMS_U64 n = 1; n <= count; n += 1, sid_ptr += 1){
-      SYMS_SymbolID sid = *sid_ptr;
-      SYMS_SymbolKind kind = syms_group_symbol_kind_from_sid(group, unit, sid);
-      if (kind == SYMS_SymbolKind_ImageRelativeVariable){
-        SYMS_U64 var_virt_off = syms_voff_from_var_sid(dbg_data, dbg, unit, sid);
-        SYMS_USID var_type_usid = syms_type_from_var_sid(dbg_data, dbg, unit, sid);
-        SYMS_U64 var_size = syms_group_type_size_from_usid(group, var_type_usid);
-        SYMS_U64Range range = {var_virt_off, var_virt_off + var_size};
-        syms_spatial_map_1d_loose_push_single(scratch.arena, &loose, sid, range);
-      }
-    }
-    
-    //- bake tight map
-    SYMS_Arena *arena = syms_group_get_lane_arena(group);
-    SYMS_SpatialMap1D map = syms_spatial_map_1d_bake(arena, &loose);
-    
-    //- save to group
-    group->unit_cache_flags[index] |= SYMS_GroupUnitCacheFlag_HasVarMap;
-    group->unit_var_maps[index] = map;
-    
-    syms_release_scratch(scratch);
-    
-    SYMS_ASSERT_PARANOID(syms_spatial_map_1d_invariants(&map));
-  }
-  
-  SYMS_SpatialMap1D *result = &syms_spatial_map_1d_nil;
-  if (index < group->unit_count){
-    result = &group->unit_var_maps[index];
-  }
-  
-  SYMS_ProfEnd();
-  return(result);
-}
-
 SYMS_API void
 syms_group_fetch_line_to_addr_maps_from_uid(SYMS_Group *group, SYMS_UnitID uid){
   SYMS_ProfBegin("syms_group_fetch_line_to_addr_maps_from_uid");
@@ -1241,10 +1167,53 @@ syms_group_line_to_addr_map_from_uid_file_id(SYMS_Group *group, SYMS_UnitID uid,
   return(result);
 }
 
+SYMS_API SYMS_SymbolNameMap*
+syms_group_type_map_from_uid(SYMS_Group *group, SYMS_UnitID uid){
+  SYMS_ProfBegin("syms_group_type_map_from_uid");
+  
+  SYMS_SymbolNameMap *result = 0;
+  
+  SYMS_U64 index = uid - 1;
+  if (index < group->unit_count){
+    if ((group->unit_cache_flags[index] & SYMS_GroupUnitCacheFlag_HasTypeNameMap) == 0){
+      SYMS_ArenaTemp scratch = syms_get_scratch(0, 0);
+      
+      // grab group members
+      SYMS_String8 dbg_data = group->dbg_data;
+      SYMS_DbgAccel *dbg = group->dbg;
+      
+      // grab unit
+      SYMS_UnitAccel *unit = syms_group_unit_from_uid(group, uid);
+      
+      // assemble loose map
+      SYMS_SymbolNameMapLoose loose = syms_symbol_name_map_begin(scratch.arena, 4001);
+      SYMS_SymbolIDArray *type_array = syms_group_type_sid_array_from_uid(group, uid);
+      SYMS_SymbolID *sid_opl = type_array->ids + type_array->count;
+      SYMS_SymbolID *sid_ptr = type_array->ids;
+      for (;sid_ptr < sid_opl; sid_ptr += 1){
+        SYMS_String8 name = syms_symbol_name_from_sid(scratch.arena, dbg_data, dbg, unit, *sid_ptr);
+        syms_symbol_name_map_push(scratch.arena, &loose, name, *sid_ptr);
+      }
+      
+      //- bake the map
+      group->unit_type_maps[index] = syms_symbol_name_map_bake(group->arena, &loose);
+      group->unit_cache_flags[index] |= SYMS_GroupUnitCacheFlag_HasTypeNameMap;
+      
+      syms_release_scratch(scratch);
+    }
+    
+    result = &group->unit_type_maps[index];
+  }
+  
+  SYMS_ProfEnd();
+  return(result);
+}
+
+
 SYMS_API SYMS_SpatialMap1D*
 syms_group_stripped_info_map(SYMS_Group *group){
-  SYMS_ProfBegin("syms_group_stripped_info_map");
   SYMS_ASSERT_PARANOID(syms_thread_lane == 0);
+  SYMS_ProfBegin("syms_group_stripped_info_map");
   if (!group->stripped_info_map_is_built){
     SYMS_Arena *arena = group->arena;
     
@@ -1353,15 +1322,6 @@ syms_group_proc_sid_from_uid_voff__accelerated(SYMS_Group *group, SYMS_UnitID ui
   return(result);
 }
 
-SYMS_API SYMS_SymbolID
-syms_group_var_sid_from_uid_voff__accelerated(SYMS_Group *group, SYMS_UnitID uid, SYMS_U64 voff){
-  SYMS_ProfBegin("syms_group_var_sid_from_uid_voff__accelerated");
-  SYMS_SpatialMap1D *map = syms_group_var_map_from_uid(group, uid);
-  SYMS_SymbolID result = syms_spatial_map_1d_value_from_point(map, voff);
-  SYMS_ProfEnd();
-  return(result);
-}
-
 SYMS_API SYMS_Line
 syms_group_line_from_uid_voff__accelerated(SYMS_Group *group, SYMS_UnitID uid, SYMS_U64 voff){
   SYMS_ProfBegin("syms_group_line_from_uid_voff__accelerated");
@@ -1500,20 +1460,136 @@ syms_line_to_addr_line_sort__rec(SYMS_FileToLineToAddrLooseLine **array, SYMS_U6
   }
 }
 
+////////////////////////////////
+// NOTE(allen): Syms Group Type Graph
+
+SYMS_API SYMS_TypeGraph* syms_group_type_graph(SYMS_Group *group){
+  SYMS_TypeGraph *result = &group->type_graph;
+  return(result);
+}
+
+////////////////////////////////
+// NOTE(allen): Syms Group Varaible Address Mapping Functions
+
+SYMS_API SYMS_SymbolID
+syms_group_var_sid_from_uid_voff__linear_scan(SYMS_TypeGraph *graph, SYMS_Group *group,
+                                              SYMS_UnitID uid, SYMS_U64 voff){
+  SYMS_ASSERT_PARANOID(syms_thread_lane == 0);
+  SYMS_ProfBegin("syms_group_var_sid_from_uid_voff__linear_scan");
+  
+  SYMS_String8 dbg_data = group->dbg_data;
+  SYMS_DbgAccel *dbg = group->dbg;
+  SYMS_UnitAccel *unit = syms_group_unit_from_uid(group, uid);
+  SYMS_SymbolIDArray *sid_array = syms_group_var_sid_array_from_uid(group, uid);
+  SYMS_U64 count = sid_array->count;
+  
+  SYMS_SymbolID result = {0};
+  SYMS_SymbolID *sid_ptr = sid_array->ids;
+  for (SYMS_U64 i = 0; i < count; i += 1, sid_ptr += 1){
+    SYMS_SymbolID sid = *sid_ptr;
+    SYMS_SymbolKind kind = syms_group_symbol_kind_from_sid(group, unit, sid);
+    SYMS_B32 hit = syms_false;
+    if (kind == SYMS_SymbolKind_ImageRelativeVariable){
+      SYMS_U64 var_voff = syms_voff_from_var_sid(dbg_data, dbg, unit, sid);
+      if (var_voff == voff){
+        hit = syms_true;
+      }
+      else if (var_voff < voff){
+        SYMS_USID var_type_usid = syms_type_from_var_sid(dbg_data, dbg, unit, sid);
+        SYMS_U64 var_size = syms_group_type_size_from_usid(graph, group, var_type_usid);
+        if (voff < var_voff + var_size){
+          hit = syms_true;
+        }
+      }
+    }
+    if (hit){
+      result = sid;
+      break;
+    }
+  }
+  
+  SYMS_ProfEnd();
+  return(result);
+}
+
+SYMS_API SYMS_SpatialMap1D*
+syms_group_var_map_from_uid(SYMS_TypeGraph *graph, SYMS_Group *group, SYMS_UnitID uid){
+  SYMS_ASSERT_PARANOID(syms_thread_lane == 0);
+  SYMS_ProfBegin("syms_group_var_map_from_uid");
+  SYMS_U64 index = uid - 1;
+  if (index < group->unit_count &&
+      (group->unit_cache_flags[index] & SYMS_GroupUnitCacheFlag_HasVarMap) == 0){
+    //- build loose map
+    SYMS_ArenaTemp scratch = syms_get_scratch(0, 0);
+    SYMS_SpatialMap1DLoose loose = {0};
+    
+    SYMS_String8 dbg_data = group->dbg_data;
+    SYMS_DbgAccel *dbg = group->dbg;
+    SYMS_UnitAccel *unit = syms_group_unit_from_uid(group, uid);
+    SYMS_SymbolIDArray *sid_array = syms_group_var_sid_array_from_uid(group, uid);
+    SYMS_U64 count = sid_array->count;
+    
+    SYMS_SymbolID *sid_ptr = sid_array->ids;
+    for (SYMS_U64 n = 1; n <= count; n += 1, sid_ptr += 1){
+      SYMS_SymbolID sid = *sid_ptr;
+      SYMS_SymbolKind kind = syms_group_symbol_kind_from_sid(group, unit, sid);
+      if (kind == SYMS_SymbolKind_ImageRelativeVariable){
+        SYMS_U64 var_virt_off = syms_voff_from_var_sid(dbg_data, dbg, unit, sid);
+        SYMS_USID var_type_usid = syms_type_from_var_sid(dbg_data, dbg, unit, sid);
+        SYMS_U64 var_size = syms_group_type_size_from_usid(graph, group, var_type_usid);
+        SYMS_U64Range range = {var_virt_off, var_virt_off + var_size};
+        syms_spatial_map_1d_loose_push_single(scratch.arena, &loose, sid, range);
+      }
+    }
+    
+    //- bake tight map
+    SYMS_Arena *arena = syms_group_get_lane_arena(group);
+    SYMS_SpatialMap1D map = syms_spatial_map_1d_bake(arena, &loose);
+    
+    //- save to group
+    group->unit_cache_flags[index] |= SYMS_GroupUnitCacheFlag_HasVarMap;
+    group->unit_var_maps[index] = map;
+    
+    syms_release_scratch(scratch);
+    
+    SYMS_ASSERT_PARANOID(syms_spatial_map_1d_invariants(&map));
+  }
+  
+  SYMS_SpatialMap1D *result = &syms_spatial_map_1d_nil;
+  if (index < group->unit_count){
+    result = &group->unit_var_maps[index];
+  }
+  
+  SYMS_ProfEnd();
+  return(result);
+}
+
+SYMS_API SYMS_SymbolID
+syms_group_var_sid_from_uid_voff__accelerated(SYMS_TypeGraph *graph, SYMS_Group *group,
+                                              SYMS_UnitID uid, SYMS_U64 voff){
+  SYMS_ASSERT_PARANOID(syms_thread_lane == 0);
+  SYMS_ProfBegin("syms_group_var_sid_from_uid_voff__accelerated");
+  SYMS_SpatialMap1D *map = syms_group_var_map_from_uid(graph, group, uid);
+  SYMS_SymbolID result = syms_spatial_map_1d_value_from_point(map, voff);
+  SYMS_ProfEnd();
+  return(result);
+}
+
+
 ////////////////////////////////
 // NOTE(allen): Syms Group Type Graph Functions
 
 SYMS_API SYMS_TypeNode*
-syms_group_type_from_usid(SYMS_Group *group, SYMS_USID usid){
-  SYMS_ProfBegin("syms_group_type_from_usid");
+syms_group_type_from_usid(SYMS_TypeGraph *graph, SYMS_Group *group, SYMS_USID usid){
   SYMS_ASSERT_PARANOID(syms_thread_lane == 0);
-  SYMS_TypeNode *result = syms_group_type_from_usid__rec(group, usid);
+  SYMS_ProfBegin("syms_group_type_from_usid");
+  SYMS_TypeNode *result = syms_group_type_from_usid__rec(graph, group, usid);
   SYMS_ProfEnd();
   return(result);
 }
 
 SYMS_API SYMS_TypeNode*
-syms_group_type_from_usid__rec(SYMS_Group *group, SYMS_USID usid){
+syms_group_type_from_usid__rec(SYMS_TypeGraph *graph, SYMS_Group *group, SYMS_USID usid){
   SYMS_TypeNode *result = 0;
   
   //- look at unit features
@@ -1522,7 +1598,7 @@ syms_group_type_from_usid__rec(SYMS_Group *group, SYMS_USID usid){
   if (can_have_types){
     
     //- get cached version of type
-    SYMS_TypeNode *type_from_cache = syms_type_from_usid(&group->type_usid_buckets, usid);
+    SYMS_TypeNode *type_from_cache = syms_type_from_usid(&graph->type_usid_buckets, usid);
     
     //- construct new graph node
     SYMS_TypeNode *new_type = 0;
@@ -1537,7 +1613,8 @@ syms_group_type_from_usid__rec(SYMS_Group *group, SYMS_USID usid){
       // is this a type?
       SYMS_SymbolKind kind = syms_group_symbol_kind_from_sid(group, unit, usid.sid);
       if (kind != SYMS_SymbolKind_Type){
-        syms_type_usid_buckets_insert(group->arena, &group->type_usid_buckets, usid, new_type);
+        syms_type_usid_buckets_insert(graph->arena, &graph->type_usid_buckets,
+                                      usid, new_type);
       }
       else{
         
@@ -1547,7 +1624,7 @@ syms_group_type_from_usid__rec(SYMS_Group *group, SYMS_USID usid){
         
         if (syms_type_kind_is_basic(type_info.kind)){
           SYMS_String8 name = syms_group_symbol_name_from_sid(scratch.arena, group, unit, usid.sid);
-          new_type = syms_group_type_basic(group, type_info.kind, type_info.reported_size, name);
+          new_type = syms_type_basic(graph, type_info.kind, type_info.reported_size, name);
         }
         else switch (type_info.kind){
           case SYMS_TypeKind_Bitfield:
@@ -1562,8 +1639,9 @@ syms_group_type_from_usid__rec(SYMS_Group *group, SYMS_USID usid){
           case SYMS_TypeKind_ForwardEnum:
           {
             // stub in the new type (necessary to do this before we infer more types to handle cycles)
-            new_type = syms_push_array_zero(group->arena, SYMS_TypeNode, 1);
-            syms_type_usid_buckets_insert(group->arena, &group->type_usid_buckets, usid, new_type);
+            new_type = syms_push_array_zero(graph->arena, SYMS_TypeNode, 1);
+            syms_type_usid_buckets_insert(graph->arena, &graph->type_usid_buckets,
+                                          usid, new_type);
             pre_inserted = syms_true;
             
             // resolve references
@@ -1589,11 +1667,11 @@ syms_group_type_from_usid__rec(SYMS_Group *group, SYMS_USID usid){
                 }
               }
               
-              referenced_type = syms_group_type_from_usid__rec(group, match_usid);
+              referenced_type = syms_group_type_from_usid__rec(graph, group, match_usid);
             }
             
             // resolve direct type
-            SYMS_TypeNode *direct = syms_group_type_from_usid__rec(group, type_info.direct_type);
+            SYMS_TypeNode *direct = syms_group_type_from_usid__rec(graph, group, type_info.direct_type);
             
             // compute size info
             SYMS_U64 byte_size = 0;
@@ -1625,7 +1703,7 @@ syms_group_type_from_usid__rec(SYMS_Group *group, SYMS_USID usid){
             }
             
             // setup src coord
-            SYMS_TypeSrcCoord *src_coord = syms_push_array_zero(group->arena, SYMS_TypeSrcCoord, 1);
+            SYMS_TypeSrcCoord *src_coord = syms_push_array_zero(graph->arena, SYMS_TypeSrcCoord, 1);
             src_coord->usid = usid;
             src_coord->file_id = type_info.src_coord.file_id;
             src_coord->line = type_info.src_coord.line;
@@ -1636,7 +1714,9 @@ syms_group_type_from_usid__rec(SYMS_Group *group, SYMS_USID usid){
             
             // fill new type
             new_type->kind = type_info.kind;
-            new_type->name = syms_string_cons(group->arena, &group->string_cons, name);
+            new_type->name = syms_string_cons(graph->arena,
+                                              graph->string_cons,
+                                              name);
             new_type->byte_size = byte_size;
             new_type->src_coord = src_coord;
             new_type->direct_type = direct_type;
@@ -1645,21 +1725,21 @@ syms_group_type_from_usid__rec(SYMS_Group *group, SYMS_USID usid){
           
           case SYMS_TypeKind_Modifier:
           {
-            SYMS_TypeNode *direct = syms_group_type_from_usid__rec(group, type_info.direct_type);
-            new_type = syms_group_type_mod_from_type(group, direct, type_info.mods);
+            SYMS_TypeNode *direct = syms_group_type_from_usid__rec(graph, group, type_info.direct_type);
+            new_type = syms_type_mod_from_type(graph, direct, type_info.mods);
           }break;
           
           case SYMS_TypeKind_Ptr:
           case SYMS_TypeKind_LValueReference:
           case SYMS_TypeKind_RValueReference:
           {
-            SYMS_TypeNode *direct = syms_group_type_from_usid__rec(group, type_info.direct_type);
-            new_type = syms_group_type_ptr_from_type(group, type_info.kind, direct);
+            SYMS_TypeNode *direct = syms_group_type_from_usid__rec(graph, group, type_info.direct_type);
+            new_type = syms_type_ptr_from_type(graph, type_info.kind, direct);
           }break;
           
           case SYMS_TypeKind_Array:
           {
-            SYMS_TypeNode *direct = syms_group_type_from_usid__rec(group, type_info.direct_type);
+            SYMS_TypeNode *direct = syms_group_type_from_usid__rec(graph, group, type_info.direct_type);
             
             SYMS_U64 array_count = type_info.reported_size;
             if (type_info.reported_size_interp == SYMS_SizeInterpretation_ByteCount){
@@ -1670,18 +1750,18 @@ syms_group_type_from_usid__rec(SYMS_Group *group, SYMS_USID usid){
               array_count = type_info.reported_size/next_size;
             }
             
-            new_type = syms_group_type_array_from_type(group, direct, array_count);
+            new_type = syms_type_array_from_type(graph, direct, array_count);
           }break;
           
           case SYMS_TypeKind_Proc:
           {
-            SYMS_TypeNode *direct = syms_group_type_from_usid__rec(group, type_info.direct_type);
+            SYMS_TypeNode *direct = syms_group_type_from_usid__rec(graph, group, type_info.direct_type);
             
             // read signature
             SYMS_SigInfo sig_info = syms_sig_info_from_type_sid(scratch.arena, dbg_data, dbg, unit, usid.sid);
             
             SYMS_U64 param_count = sig_info.param_type_ids.count;
-            SYMS_TypeNode **params = syms_push_array(group->arena, SYMS_TypeNode*, param_count);
+            SYMS_TypeNode **params = syms_push_array(graph->arena, SYMS_TypeNode*, param_count);
             
             {
               SYMS_TypeNode **param_ptr = params;
@@ -1689,21 +1769,21 @@ syms_group_type_from_usid__rec(SYMS_Group *group, SYMS_USID usid){
               SYMS_SymbolID *sid_ptr = sig_info.param_type_ids.ids;
               for (;param_ptr < param_opl; sid_ptr += 1, param_ptr += 1){
                 SYMS_USID param_usid = syms_make_usid(sig_info.uid, *sid_ptr);
-                *param_ptr = syms_group_type_from_usid__rec(group, param_usid);
+                *param_ptr = syms_group_type_from_usid__rec(graph, group, param_usid);
               }
             }
             
             SYMS_USID this_usid = syms_make_usid(sig_info.uid, sig_info.this_type_id);
-            SYMS_TypeNode *this_type = syms_group_type_from_usid__rec(group, this_usid);
+            SYMS_TypeNode *this_type = syms_group_type_from_usid__rec(graph, group, this_usid);
             
-            new_type = syms_group_type_proc_from_type(group, direct, this_type, params, param_count);
+            new_type = syms_type_proc_from_type(graph, direct, this_type, params, param_count);
           }break;
           
           case SYMS_TypeKind_MemberPtr:
           {
-            SYMS_TypeNode *direct = syms_group_type_from_usid__rec(group, type_info.direct_type);
-            SYMS_TypeNode *container = syms_group_type_from_usid__rec(group, type_info.containing_type);
-            new_type = syms_group_type_member_ptr_from_type(group, container, direct);
+            SYMS_TypeNode *direct = syms_group_type_from_usid__rec(graph, group, type_info.direct_type);
+            SYMS_TypeNode *container = syms_group_type_from_usid__rec(graph, group, type_info.containing_type);
+            new_type = syms_type_member_ptr_from_type(graph, container, direct);
             if (type_info.reported_size_interp == SYMS_SizeInterpretation_ByteCount){
               new_type->byte_size = type_info.reported_size;
             }
@@ -1723,11 +1803,12 @@ syms_group_type_from_usid__rec(SYMS_Group *group, SYMS_USID usid){
         //- split out modifiers into stand alone nodes
         if (type_info.mods != 0 && type_info.kind != SYMS_TypeKind_Modifier){
           SYMS_TypeNode *direct = new_type;
-          new_type = syms_group_type_mod_from_type(group, direct, type_info.mods);
+          new_type = syms_type_mod_from_type(graph, direct, type_info.mods);
         }
         
         if (!pre_inserted){
-          syms_type_usid_buckets_insert(group->arena, &group->type_usid_buckets, usid, new_type);
+          syms_type_usid_buckets_insert(graph->arena, &graph->type_usid_buckets,
+                                        usid, new_type);
         }
       }
       
@@ -1748,151 +1829,26 @@ syms_group_type_from_usid__rec(SYMS_Group *group, SYMS_USID usid){
   return(result);
 }
 
-SYMS_API SYMS_TypeMemberArray*
-syms_group_members_from_type(SYMS_Group *group, SYMS_TypeNode *type){
-  SYMS_ProfBegin("syms_group_members_from_type");
-  SYMS_TypeMemberArray *result = &syms_type_member_array_nil;
-  switch (type->kind){
-    case SYMS_TypeKind_Class:
-    case SYMS_TypeKind_Struct:
-    case SYMS_TypeKind_Union:
-    {
-      if (type->lazy_ptr == 0 && type->src_coord != 0){
-        SYMS_USID usid = type->src_coord->usid;
-        SYMS_String8 dbg_data = group->dbg_data;
-        SYMS_DbgAccel *dbg = group->dbg;
-        
-        SYMS_UnitAccel *unit = syms_group_unit_from_uid(group, usid.uid);
-        
-        SYMS_ArenaTemp scratch = syms_get_scratch(0, 0);
-        SYMS_MemsAccel *mems_accel = syms_mems_accel_from_sid(scratch.arena, dbg_data, dbg, unit, usid.sid);
-        SYMS_U64 mem_count = syms_mem_count_from_mems(mems_accel);
-        SYMS_TypeMember *mems = syms_push_array(group->arena, SYMS_TypeMember, mem_count);
-        
-        SYMS_TypeMember *mem_ptr = mems;
-        for (SYMS_U64 n = 1; n <= mem_count; n += 1, mem_ptr += 1){
-          SYMS_ArenaTemp temp = syms_arena_temp_begin(scratch.arena);
-          
-          SYMS_MemInfo mem_info = syms_mem_info_from_number(group->arena, dbg_data, dbg, unit, mems_accel, n);
-          
-          SYMS_TypeNode *type = &syms_type_graph_nil;
-          switch (mem_info.kind){
-            case SYMS_MemKind_VTablePtr: /*no type*/ break;
-            
-            case SYMS_MemKind_DataField:
-            case SYMS_MemKind_StaticData:
-            {
-              // TODO(allen): here we lose access to the fast path from StaticData to a virtual offset.
-              // This path is only available from the DWARF backend, but there are a few ways we could
-              // preserve it.
-              SYMS_USID type_usid = syms_type_from_mem_number(dbg_data, dbg, unit, mems_accel, n);
-              type = syms_group_type_from_usid__rec(group, type_usid);
-            }break;
-            
-            case SYMS_MemKind_BaseClass:
-            case SYMS_MemKind_VBaseClassPtr:
-            case SYMS_MemKind_NestedType:
-            {
-              // directly get type from member id
-              SYMS_USID type_usid = syms_type_from_mem_number(dbg_data, dbg, unit, mems_accel, n);
-              type = syms_group_type_from_usid__rec(group, type_usid);
-            }break;
-            
-            case SYMS_MemKind_Method:
-            case SYMS_MemKind_StaticMethod:
-            {
-              // TODO(allen): here we lose access to the fast path a procedure symbol.
-              // This path is only available from the DWARF backend, but there are a few ways we could
-              // preserve it.
-              SYMS_SigInfo sig = syms_sig_info_from_mem_number(scratch.arena, dbg_data, dbg, unit, mems_accel, n);
-              
-              SYMS_USID return_usid = syms_make_usid(sig.uid, sig.return_type_id);
-              SYMS_USID this_usid = syms_make_usid(sig.uid, sig.this_type_id);
-              
-              SYMS_TypeNode *ret_type  = syms_group_type_from_usid__rec(group, return_usid);
-              SYMS_TypeNode *this_type = syms_group_type_from_usid__rec(group, this_usid);
-              SYMS_U64 param_count = sig.param_type_ids.count;
-              SYMS_TypeNode **params = syms_push_array(group->arena, SYMS_TypeNode*, param_count);
-              
-              {
-                SYMS_TypeNode **param_ptr = params;
-                SYMS_TypeNode **param_opl = params + sig.param_type_ids.count;
-                SYMS_SymbolID *param_id_ptr = sig.param_type_ids.ids;
-                for (; param_ptr < param_opl; param_id_ptr += 1, param_ptr += 1){
-                  *param_ptr = syms_group_type_from_usid__rec(group, syms_make_usid(sig.uid, *param_id_ptr));
-                }
-              }
-              
-              type = syms_group_type_proc_from_type(group, ret_type, this_type, params, param_count);
-            }break;
-          }
-          
-          mem_ptr->kind = mem_info.kind;
-          mem_ptr->visibility = mem_info.visibility;
-          mem_ptr->flags = mem_info.flags;
-          mem_ptr->name = syms_string_cons(group->arena, &group->string_cons, mem_info.name);
-          mem_ptr->off = mem_info.off;
-          mem_ptr->virtual_off = mem_info.virtual_off;
-          mem_ptr->type = type;
-          
-          syms_arena_temp_end(temp);
-        }
-        
-        syms_release_scratch(scratch);
-        
-        
-        SYMS_TypeMemberArray *type_member_array = syms_push_array(group->arena, SYMS_TypeMemberArray, 1);
-        type_member_array->mems = mems;
-        type_member_array->count = mem_count;
-        type->lazy_ptr = type_member_array;
-      }
-      
-      result = (SYMS_TypeMemberArray*)type->lazy_ptr;
-    }break;
-  }
-  SYMS_ProfEnd();
-  return(result);
-}
-
-SYMS_API SYMS_EnumInfoArray*
-syms_group_enum_members_from_type(SYMS_Group *group, SYMS_TypeNode *type){
-  SYMS_ProfBegin("syms_group_enum_members_from_type");
-  SYMS_EnumInfoArray *result = &syms_enum_info_array_nil;
-  if (type->kind == SYMS_TypeKind_Enum){
-    if (type->lazy_ptr == 0 && type->src_coord != 0){
-      SYMS_USID usid = type->src_coord->usid;
-      SYMS_UnitAccel *unit = syms_group_unit_from_uid(group, usid.uid);
-      SYMS_EnumInfoArray *enum_info_array = syms_push_array(group->arena, SYMS_EnumInfoArray, 1);
-      *enum_info_array = syms_enum_info_array_from_sid(group->arena, group->dbg_data, group->dbg,
-                                                       unit, usid.sid);
-      type->lazy_ptr = enum_info_array;
-    }
-    result = (SYMS_EnumInfoArray*)type->lazy_ptr;
-  }
-  SYMS_ProfEnd();
-  return(result);
-}
-
-SYMS_API SYMS_U64
-syms_group_type_size_from_usid(SYMS_Group *group, SYMS_USID usid){
-  SYMS_ProfBegin("syms_group_type_size_from_usid");
-  SYMS_ASSERT_PARANOID(syms_thread_lane == 0);
-  SYMS_UnitAccel *unit = syms_group_unit_from_uid(group, usid.uid);
-  SYMS_TypeInfo type_info = syms_type_info_from_sid(group->dbg_data, group->dbg, unit, usid.sid);
-  SYMS_U64 result = 0;
-  switch (type_info.reported_size_interp){
-    case SYMS_SizeInterpretation_ByteCount:
+SYMS_API SYMS_U64
+syms_group_type_size_from_usid(SYMS_TypeGraph *graph, SYMS_Group *group, SYMS_USID usid){
+  SYMS_ASSERT_PARANOID(syms_thread_lane == 0);
+  SYMS_ProfBegin("syms_group_type_size_from_usid");
+  SYMS_UnitAccel *unit = syms_group_unit_from_uid(group, usid.uid);
+  SYMS_TypeInfo type_info = syms_type_info_from_sid(group->dbg_data, group->dbg, unit, usid.sid);
+  SYMS_U64 result = 0;
+  switch (type_info.reported_size_interp){
+    case SYMS_SizeInterpretation_ByteCount:
     {
       result = type_info.reported_size;
     }break;
     case SYMS_SizeInterpretation_Multiplier:
     {
-      SYMS_U64 element_size = syms_group_type_size_from_usid(group, type_info.direct_type);
+      SYMS_U64 element_size = syms_group_type_size_from_usid(graph, group, type_info.direct_type);
       result = element_size*type_info.reported_size;
     }break;
     default:
     {
-      SYMS_TypeNode *type_node = syms_group_type_from_usid(group, usid);
+      SYMS_TypeNode *type_node = syms_group_type_from_usid(graph, group, usid);
       result = type_node->byte_size;
     }break;
   }
@@ -1900,490 +1856,175 @@ syms_group_type_size_from_usid(SYMS_Group *group, SYMS_USID usid){
   return(result);
 }
 
-SYMS_API SYMS_TypeNode*
-syms_group_type_basic(SYMS_Group *group, SYMS_TypeKind basic_kind, SYMS_U64 size, SYMS_String8 name){
-  SYMS_ProfBegin("syms_group_type_basic");
-  SYMS_ASSERT_PARANOID(syms_thread_lane == 0);
-  
-  SYMS_TypeNode *result = &syms_type_graph_nil;
-  
-  SYMS_String8 name_cons = syms_string_cons(group->arena, &group->string_cons, name);
-  
-  SYMS_U64 content[3];
-  content[0] = basic_kind;
-  content[1] = size;
-  content[2] = (SYMS_U64)name_cons.str;
-  SYMS_String8 data = {(SYMS_U8*)content, sizeof(content)};
-  result = syms_type_from_content(&group->type_content_buckets, data);
-  
-  if (result == 0){
-    result = syms_push_array_zero(group->arena, SYMS_TypeNode, 1);
-    syms_type_content_buckets_insert(group->arena, &group->type_content_buckets, data, result);
-    
-    result->kind = basic_kind;
-    result->name = name_cons;
-    result->byte_size = size;
-    result->direct_type = &syms_type_graph_nil;
-    result->this_type = &syms_type_graph_nil;
-  }
-  
-  SYMS_ProfEnd();
+SYMS_API SYMS_TypeChain syms_group_artificial_types_from_name(SYMS_Group *group,
+                                                              SYMS_String8 name){
+  SYMS_TypeChain result = syms_type_from_name(&group->type_graph, name);
   return(result);
 }
 
-SYMS_API SYMS_TypeNode*
-syms_group_type_mod_from_type(SYMS_Group *group, SYMS_TypeNode *type, SYMS_TypeModifiers mods){
-  SYMS_ProfBegin("syms_group_type_mod_from_type");
-  SYMS_ASSERT_PARANOID(syms_thread_lane == 0);
-  
-  SYMS_TypeNode *result = &syms_type_graph_nil;
-  
-  SYMS_U64 content[3];
-  content[0] = SYMS_TypeKind_Modifier;
-  content[1] = mods;
-  content[2] = (SYMS_U64)type;
-  SYMS_String8 data = {(SYMS_U8*)content, sizeof(content)};
-  result = syms_type_from_content(&group->type_content_buckets, data);
+SYMS_API SYMS_USIDList
+syms_group_type_list_from_name_accelerated(SYMS_Arena *arena, SYMS_Group *group, SYMS_String8 name){
+  SYMS_USIDList result = {0};
   
-  if (result == 0){
-    result = syms_push_array_zero(group->arena, SYMS_TypeNode, 1);
-    syms_type_content_buckets_insert(group->arena, &group->type_content_buckets, data, result);
-    
-    result->kind = SYMS_TypeKind_Modifier;
-    result->byte_size = type->byte_size;
-    result->direct_type = type;
-    result->this_type = &syms_type_graph_nil;
-    result->mods = mods;
-  }
-  
-  SYMS_ProfEnd();
-  return(result);
-}
-
-SYMS_API SYMS_TypeNode*
-syms_group_type_ptr_from_type(SYMS_Group *group, SYMS_TypeKind ptr_kind, SYMS_TypeNode *type){
-  SYMS_ProfBegin("syms_group_type_ptr_from_type");
-  SYMS_ASSERT_PARANOID(syms_thread_lane == 0);
-  
-  SYMS_TypeNode *result = &syms_type_graph_nil;
-  
-  SYMS_U64 content[2];
-  content[0] = ptr_kind;
-  content[1] = (SYMS_U64)type;
-  SYMS_String8 data = {(SYMS_U8*)content, sizeof(content)};
-  result = syms_type_from_content(&group->type_content_buckets, data);
-  
-  if (result == 0){
-    result = syms_push_array_zero(group->arena, SYMS_TypeNode, 1);
-    syms_type_content_buckets_insert(group->arena, &group->type_content_buckets, data, result);
-    
-    result->kind = ptr_kind;
-    result->byte_size = syms_group_address_size(group);
-    result->direct_type = type;
-    result->this_type = &syms_type_graph_nil;
-  }
-  
-  SYMS_ProfEnd();
-  return(result);
-}
-
-SYMS_API SYMS_TypeNode*
-syms_group_type_array_from_type(SYMS_Group *group, SYMS_TypeNode *type, SYMS_U64 count){
-  SYMS_ProfBegin("syms_group_type_array_from_type");
-  SYMS_ASSERT_PARANOID(syms_thread_lane == 0);
-  
-  SYMS_TypeNode *result = &syms_type_graph_nil;
-  
-  SYMS_U64 content[3];
-  content[0] = SYMS_TypeKind_Array;
-  content[1] = count;
-  content[2] = (SYMS_U64)type;
-  SYMS_String8 data = {(SYMS_U8*)content, sizeof(content)};
-  result = syms_type_from_content(&group->type_content_buckets, data);
-  
-  if (result == 0){
-    result = syms_push_array_zero(group->arena, SYMS_TypeNode, 1);
-    syms_type_content_buckets_insert(group->arena, &group->type_content_buckets, data, result);
-    
-    result->kind = SYMS_TypeKind_Array;
-    result->byte_size = count*type->byte_size;
-    result->direct_type = type;
-    result->this_type = &syms_type_graph_nil;
-    result->array_count = count;
+  SYMS_MapAndUnit map_unit = syms_group_type_map(group);
+  if (syms_accel_is_good(map_unit.map)){
+    SYMS_String8 dbg_data = group->dbg_data;
+    SYMS_DbgAccel *dbg = group->dbg;
+    result = syms_symbol_list_from_string(arena, dbg_data, dbg,
+                                          map_unit, name);
   }
   
-  SYMS_ProfEnd();
   return(result);
 }
 
-SYMS_API SYMS_TypeNode*
-syms_group_type_proc_from_type(SYMS_Group *group, SYMS_TypeNode *ret_type, SYMS_TypeNode *this_type,
-                               SYMS_TypeNode **param_types, SYMS_U64 param_count){
-  SYMS_ProfBegin("syms_group_type_proc_from_type");
-  SYMS_ASSERT_PARANOID(syms_thread_lane == 0);
-  
-  SYMS_TypeNode *result = &syms_type_graph_nil;
-  
-  SYMS_String8 params_as_string = {(SYMS_U8*)param_types, sizeof(*param_types)*param_count};
-  SYMS_String8 params_cons = syms_string_cons(group->arena, &group->string_cons, params_as_string);
-  
-  SYMS_U64 content[4];
-  content[0] = SYMS_TypeKind_Proc;
-  content[1] = (SYMS_U64)ret_type;
-  content[2] = (SYMS_U64)this_type;
-  content[3] = (SYMS_U64)params_cons.str;
-  SYMS_String8 data = {(SYMS_U8*)content, sizeof(content)};
-  result = syms_type_from_content(&group->type_content_buckets, data);
-  
-  if (result == 0){
-    result = syms_push_array_zero(group->arena, SYMS_TypeNode, 1);
-    syms_type_content_buckets_insert(group->arena, &group->type_content_buckets, data, result);
-    
-    result->kind = SYMS_TypeKind_Proc;
-    result->byte_size = 0;
-    result->direct_type = ret_type;
-    result->this_type = this_type;
-    result->proc.params = (SYMS_TypeNode**)(params_cons.str);
-    result->proc.param_count = param_count;
+SYMS_API SYMS_SymbolIDArray
+syms_group_types_from_unit_name(SYMS_Group *group, SYMS_UnitID uid, SYMS_String8 name){
+  SYMS_SymbolNameMap *map = syms_group_type_map_from_uid(group, uid);
+  SYMS_SymbolIDArray result = {0};
+  if (map != 0){
+    result = syms_symbol_name_map_array_from_string(map, name);
   }
-  
-  SYMS_ProfEnd();
   return(result);
 }
 
-SYMS_API SYMS_TypeNode*
-syms_group_type_member_ptr_from_type(SYMS_Group *group, SYMS_TypeNode *container, SYMS_TypeNode *type){
-  SYMS_ProfBegin("syms_group_type_member_ptr_from_type");
-  SYMS_ASSERT_PARANOID(syms_thread_lane == 0);
-  
-  SYMS_TypeNode *result = &syms_type_graph_nil;
-  
-  SYMS_U64 content[3];
-  content[0] = SYMS_TypeKind_MemberPtr;
-  content[1] = (SYMS_U64)container;
-  content[2] = (SYMS_U64)type;
-  SYMS_String8 data = {(SYMS_U8*)content, sizeof(content)};
-  result = syms_type_from_content(&group->type_content_buckets, data);
+SYMS_API SYMS_TypeMemberArray*
+syms_type_members_from_type(SYMS_Group *group,
+                            SYMS_TypeNode *type){
+  SYMS_ProfBegin("syms_type_members_from_type");
   
-  if (result == 0){
-    result = syms_push_array_zero(group->arena, SYMS_TypeNode, 1);
-    syms_type_content_buckets_insert(group->arena, &group->type_content_buckets, data, result);
-    
-    SYMS_U64 byte_size = syms_group_address_size(group);
-    if (type->kind == SYMS_TypeKind_Proc){
-      byte_size *= 2;
-    }
-    
-    result->kind = SYMS_TypeKind_MemberPtr;
-    result->byte_size = byte_size;
-    result->direct_type = type;
-    result->this_type = container;
-  }
+  SYMS_TypeMemberArray *result = &syms_type_member_array_nil;
   
-  SYMS_ProfEnd();
-  return(result);
-}
-
-SYMS_API SYMS_String8
-syms_group_string_from_type(SYMS_Arena *arena, SYMS_TypeNode *type){
-  SYMS_ProfBegin("syms_group_string_from_type");
-  SYMS_ArenaTemp scratch = syms_get_scratch(&arena, 1);
-  SYMS_String8List list = {0};
-  syms_group_lhs_string_from_type(scratch.arena, type, &list);
-  syms_group_rhs_string_from_type(scratch.arena, type, &list);
-  SYMS_String8 result = syms_string_list_join(arena, &list, 0);
-  syms_release_scratch(scratch);
-  SYMS_ProfEnd();
-  return(result);
-}
-
-SYMS_API void
-syms_group_lhs_string_from_type__internal(SYMS_Arena *arena, SYMS_TypeNode *type, SYMS_String8List *out,
-                                          SYMS_U32 precedence_level, SYMS_B32 skip_return){
-  switch (type->kind){
-    default:
-    {
-      syms_string_list_push(arena, out, type->name);
-      syms_string_list_push(arena, out, syms_str8_lit(" "));
-    }break;
-    
-    case SYMS_TypeKind_Bitfield:
-    {
-      syms_group_lhs_string_from_type__internal(arena, type->direct_type, out, precedence_level, skip_return);
-    }break;
-    
-    case SYMS_TypeKind_Modifier:
-    {
-      syms_group_lhs_string_from_type__internal(arena, type->direct_type, out, 1, skip_return);
-      if (type->mods & SYMS_TypeModifier_Const){
-        syms_string_list_push(arena, out, syms_str8_lit("const "));
-      }
-      if (type->mods & SYMS_TypeModifier_Restrict){
-        syms_string_list_push(arena, out, syms_str8_lit("restrict "));
-      }
-      if (type->mods & SYMS_TypeModifier_Volatile){
-        syms_string_list_push(arena, out, syms_str8_lit("volatile "));
-      }
-    }break;
-    
-    case SYMS_TypeKind_Variadic:
-    {
-      syms_string_list_push(arena, out, syms_str8_lit("..."));
-    }break;
-    
-    case SYMS_TypeKind_Label:
-    {
-      // TODO(allen): ???
-    }break;
-    
-    case SYMS_TypeKind_Struct:
-    case SYMS_TypeKind_Union:
-    case SYMS_TypeKind_Enum:
-    case SYMS_TypeKind_Class:
-    case SYMS_TypeKind_Typedef:
-    {
-      syms_string_list_push(arena, out, type->name);
-      syms_string_list_push(arena, out, syms_str8_lit(" "));
-    }break;
-    
-    case SYMS_TypeKind_ForwardStruct:
-    case SYMS_TypeKind_ForwardUnion:
-    case SYMS_TypeKind_ForwardEnum:
-    case SYMS_TypeKind_ForwardClass:
-    {
-      SYMS_String8 keyword = syms_string_from_enum_value(SYMS_TypeKind, type->kind);
-      syms_string_list_push(arena, out, keyword);
-      syms_string_list_push(arena, out, syms_str8_lit(" "));
-      syms_string_list_push(arena, out, type->name);
-      syms_string_list_push(arena, out, syms_str8_lit(" "));
-    }break;
-    
-    case SYMS_TypeKind_Array:
-    {
-      syms_group_lhs_string_from_type__internal(arena, type->direct_type, out, 2, skip_return);
-      if (precedence_level == 1){
-        syms_string_list_push(arena, out, syms_str8_lit("("));
-      }
-    }break;
-    
-    case SYMS_TypeKind_Proc:
-    {
-      if (!skip_return){
-        syms_group_lhs_string_from_type__internal(arena, type->direct_type, out, 2, syms_false);
-      }
-      if (precedence_level == 1){
-        syms_string_list_push(arena, out, syms_str8_lit("("));
-      }
-    }break;
-    
-    case SYMS_TypeKind_Ptr:
-    {
-      syms_group_lhs_string_from_type__internal(arena, type->direct_type, out, 1, skip_return);
-      syms_string_list_push(arena, out, syms_str8_lit("*"));
-    }break;
-    
-    case SYMS_TypeKind_LValueReference:
-    {
-      syms_group_lhs_string_from_type__internal(arena, type->direct_type, out, 1, skip_return);
-      syms_string_list_push(arena, out, syms_str8_lit("&"));
-    }break;
-    
-    case SYMS_TypeKind_RValueReference:
-    {
-      syms_group_lhs_string_from_type__internal(arena, type->direct_type, out, 1, skip_return);
-      syms_string_list_push(arena, out, syms_str8_lit("&&"));
-    }break;
-    
-    case SYMS_TypeKind_MemberPtr:
-    {
-      syms_group_lhs_string_from_type__internal(arena, type->direct_type, out, 1, skip_return);
-      SYMS_TypeNode *container = type->this_type;
-      if (container != 0){
-        syms_string_list_push(arena, out, container->name);
-      }
-      else{
-        syms_string_list_push(arena, out, syms_str8_lit("<unknown-class>"));
-      }
-      syms_string_list_push(arena, out, syms_str8_lit("::*"));
-    }break;
-  }
-}
-
-SYMS_API void
-syms_group_rhs_string_from_type__internal(SYMS_Arena *arena, SYMS_TypeNode *type, SYMS_String8List *out,
-                                          SYMS_U32 precedence_level){
-  switch (type->kind){
-    default:break;
-    
-    case SYMS_TypeKind_Bitfield:
-    {
-      syms_group_rhs_string_from_type__internal(arena, type->direct_type, out, precedence_level);
-    }break;
-    
-    case SYMS_TypeKind_Modifier:
-    case SYMS_TypeKind_Ptr:
-    case SYMS_TypeKind_LValueReference:
-    case SYMS_TypeKind_RValueReference:
-    case SYMS_TypeKind_MemberPtr:
-    {
-      syms_group_rhs_string_from_type__internal(arena, type->direct_type, out, 1);
-    }break;
-    
-    case SYMS_TypeKind_Array:
-    {
-      if (precedence_level == 1){
-        syms_string_list_push(arena, out, syms_str8_lit(")"));
-      }
+  SYMS_B32 has_members = (type->kind == SYMS_TypeKind_Class ||
+                          type->kind == SYMS_TypeKind_Struct ||
+                          type->kind == SYMS_TypeKind_Union);
+  if (has_members){
+    // fill cache
+    if (type->lazy_ptr == 0 && type->src_coord != 0){
+      SYMS_TypeGraph *graph = syms_group_type_graph(group);
       
-      SYMS_String8 count_str = syms_string_from_u64(arena, type->array_count);
-      syms_string_list_push(arena, out, syms_str8_lit("["));
-      syms_string_list_push(arena, out, count_str);
-      syms_string_list_push(arena, out, syms_str8_lit("]"));
+      SYMS_USID usid = type->src_coord->usid;
+      SYMS_String8 dbg_data = group->dbg_data;
+      SYMS_DbgAccel *dbg = group->dbg;
       
-      syms_group_rhs_string_from_type__internal(arena, type->direct_type, out, 2);
-    }break;
-    
-    case SYMS_TypeKind_Proc:
-    {
-      if (precedence_level == 1){
-        syms_string_list_push(arena, out, syms_str8_lit(")"));
-      }
+      SYMS_UnitAccel *unit = syms_group_unit_from_uid(group, usid.uid);
       
-      // parameters
-      if (type->proc.param_count == 0){
-        syms_string_list_push(arena, out, syms_str8_lit("(void)"));
-      }
-      else{
-        syms_string_list_push(arena, out, syms_str8_lit("("));
-        SYMS_U64 param_count = type->proc.param_count;
-        SYMS_TypeNode **param = type->proc.params;
-        for (SYMS_U64 i = 0; i < param_count; i += 1, param += 1){
-          SYMS_String8 param_str = syms_group_string_from_type(arena, *param);
-          SYMS_String8 param_str_trimmed = syms_str8_skip_chop_whitespace(param_str);
-          syms_string_list_push(arena, out, param_str_trimmed);
-          if (i + 1 < param_count){
-            syms_string_list_push(arena, out, syms_str8_lit(", "));
-          }
+      SYMS_ArenaTemp scratch = syms_get_scratch(0, 0);
+      SYMS_MemsAccel *mems_accel = syms_mems_accel_from_sid(scratch.arena, dbg_data, dbg, unit, usid.sid);
+      SYMS_U64 mem_count = syms_mem_count_from_mems(mems_accel);
+      SYMS_TypeMember *mems = syms_push_array(group->arena, SYMS_TypeMember, mem_count);
+      
+      SYMS_TypeMember *mem_ptr = mems;
+      for (SYMS_U64 n = 1; n <= mem_count; n += 1, mem_ptr += 1){
+        SYMS_ArenaTemp temp = syms_arena_temp_begin(scratch.arena);
+        
+        SYMS_MemInfo mem_info = syms_mem_info_from_number(group->arena, dbg_data, dbg, unit, mems_accel, n);
+        
+        SYMS_TypeNode *type = &syms_type_graph_nil;
+        switch (mem_info.kind){
+          case SYMS_MemKind_VTablePtr: /*no type*/ break;
+          
+          case SYMS_MemKind_DataField:
+          case SYMS_MemKind_StaticData:
+          {
+            // TODO(allen): here we lose access to the fast path from StaticData to a virtual offset.
+            // This path is only available from the DWARF backend, but there are a few ways we could
+            // preserve it.
+            SYMS_USID type_usid = syms_type_from_mem_number(dbg_data, dbg, unit, mems_accel, n);
+            type = syms_group_type_from_usid__rec(graph, group, type_usid);
+          }break;
+          
+          case SYMS_MemKind_BaseClass:
+          case SYMS_MemKind_VBaseClassPtr:
+          case SYMS_MemKind_NestedType:
+          {
+            // directly get type from member id
+            SYMS_USID type_usid = syms_type_from_mem_number(dbg_data, dbg, unit, mems_accel, n);
+            type = syms_group_type_from_usid__rec(graph, group, type_usid);
+          }break;
+          
+          case SYMS_MemKind_Method:
+          case SYMS_MemKind_StaticMethod:
+          {
+            // TODO(allen): here we lose access to the fast path a procedure symbol.
+            // This path is only available from the DWARF backend, but there are a few ways we could
+            // preserve it.
+            SYMS_SigInfo sig = syms_sig_info_from_mem_number(scratch.arena, dbg_data, dbg, unit, mems_accel, n);
+            
+            SYMS_USID return_usid = syms_make_usid(sig.uid, sig.return_type_id);
+            SYMS_USID this_usid = syms_make_usid(sig.uid, sig.this_type_id);
+            
+            SYMS_TypeNode *ret_type  = syms_group_type_from_usid__rec(graph, group, return_usid);
+            SYMS_TypeNode *this_type = syms_group_type_from_usid__rec(graph, group, this_usid);
+            SYMS_U64 param_count = sig.param_type_ids.count;
+            SYMS_TypeNode **params = syms_push_array(group->arena, SYMS_TypeNode*, param_count);
+            
+            {
+              SYMS_TypeNode **param_ptr = params;
+              SYMS_TypeNode **param_opl = params + sig.param_type_ids.count;
+              SYMS_SymbolID *param_id_ptr = sig.param_type_ids.ids;
+              for (; param_ptr < param_opl; param_id_ptr += 1, param_ptr += 1){
+                *param_ptr = syms_group_type_from_usid__rec(graph, group,
+                                                            syms_make_usid(sig.uid, *param_id_ptr));
+              }
+            }
+            
+            type = syms_type_proc_from_type(graph, ret_type, this_type, params, param_count);
+          }break;
         }
-        syms_string_list_push(arena, out, syms_str8_lit(")"));
+        
+        mem_ptr->kind = mem_info.kind;
+        mem_ptr->visibility = mem_info.visibility;
+        mem_ptr->flags = mem_info.flags;
+        mem_ptr->name = syms_string_cons(group->arena, &group->string_cons, mem_info.name);
+        mem_ptr->off = mem_info.off;
+        mem_ptr->virtual_off = mem_info.virtual_off;
+        mem_ptr->type = type;
+        
+        syms_arena_temp_end(temp);
       }
       
-      syms_group_rhs_string_from_type__internal(arena, type->direct_type, out, 2);
-    }break;
-    
-    case SYMS_TypeKind_Label:
-    {
-      // TODO(allen): ???
-    }break;
-  }
-}
-
-SYMS_API void
-syms_group_lhs_string_from_type(SYMS_Arena *arena, SYMS_TypeNode *type, SYMS_String8List *out){
-  SYMS_ProfBegin("syms_group_lhs_string_from_type");
-  syms_group_lhs_string_from_type__internal(arena, type, out, 0, syms_false);
-  SYMS_ProfEnd();
-}
-
-SYMS_API void
-syms_group_rhs_string_from_type(SYMS_Arena *arena, SYMS_TypeNode *type, SYMS_String8List *out){
-  SYMS_ProfBegin("syms_group_rhs_string_from_type");
-  syms_group_rhs_string_from_type__internal(arena, type, out, 0);
-  SYMS_ProfEnd();
-}
-
-SYMS_API void
-syms_group_lhs_string_from_type_skip_return(SYMS_Arena *arena, SYMS_TypeNode *type, SYMS_String8List *out){
-  SYMS_ProfBegin("syms_group_lhs_string_from_type_skip_return");
-  syms_group_lhs_string_from_type__internal(arena, type, out, 0, syms_true);
-  SYMS_ProfEnd();
-}
-
-SYMS_API SYMS_U64
-syms_type_usid_hash(SYMS_USID usid){
-  SYMS_U64 result = syms_hash_u64(usid.sid + usid.uid*97);
-  return(result);
-}
-
-SYMS_API SYMS_TypeNode*
-syms_type_from_usid(SYMS_TypeUSIDBuckets *buckets, SYMS_USID usid){
-  SYMS_ProfBegin("syms_type_from_usid");
-  SYMS_TypeNode *result = 0;
-  SYMS_U64 hash = syms_type_usid_hash(usid);
-  SYMS_U64 bucket_index = hash%buckets->bucket_count;
-  for (SYMS_TypeUSIDNode *node = buckets->buckets[bucket_index];
-       node != 0;
-       node = node->next){
-    SYMS_USID key = node->key;
-    if (key.uid == usid.uid && key.sid == usid.sid){
-      result = node->type;
-      break;
+      syms_release_scratch(scratch);
+      
+      
+      SYMS_TypeMemberArray *type_member_array = syms_push_array(group->arena, SYMS_TypeMemberArray, 1);
+      type_member_array->mems = mems;
+      type_member_array->count = mem_count;
+      type->lazy_ptr = type_member_array;
     }
+    
+    // fill result
+    result = (SYMS_TypeMemberArray*)type->lazy_ptr;
   }
-  SYMS_ProfEnd();
-  return(result);
-}
-
-SYMS_API void
-syms_type_usid_buckets_insert(SYMS_Arena *arena, SYMS_TypeUSIDBuckets *buckets, SYMS_USID key,
-                              SYMS_TypeNode *type){
-  SYMS_ProfBegin("syms_type_usid_buckets_insert");
-  SYMS_U64 hash = syms_type_usid_hash(key);
-  SYMS_U64 bucket_index = hash%buckets->bucket_count;
-  
-  SYMS_TypeUSIDNode *new_node = syms_push_array(arena, SYMS_TypeUSIDNode, 1);
-  SYMS_StackPush(buckets->buckets[bucket_index], new_node);
-  new_node->key = key;
-  new_node->type = type;
   
   SYMS_ProfEnd();
-}
-
-SYMS_API SYMS_U64
-syms_type_content_hash(SYMS_String8 data){
-  SYMS_U64 result = syms_hash_djb2(data);
   return(result);
 }
 
-SYMS_API SYMS_TypeNode*
-syms_type_from_content(SYMS_TypeContentBuckets *buckets, SYMS_String8 data){
-  SYMS_ProfBegin("syms_type_from_content");
-  SYMS_TypeNode *result = 0;
-  SYMS_U64 hash = syms_type_content_hash(data);
-  SYMS_U64 bucket_index = hash%buckets->bucket_count;
-  for (SYMS_TypeContentNode *node = buckets->buckets[bucket_index];
-       node != 0;
-       node = node->next){
-    if (node->hash == hash && syms_string_match(node->key, data, 0)){
-      result = node->type;
-      break;
+SYMS_API SYMS_EnumInfoArray*
+syms_type_enum_members_from_type(SYMS_Group *group,
+                                 SYMS_TypeNode *type){
+  SYMS_ProfBegin("syms_type_enum_members_from_type");
+  SYMS_EnumInfoArray *result = &syms_enum_info_array_nil;
+  if (type->kind == SYMS_TypeKind_Enum){
+    // fill cache
+    if (type->lazy_ptr == 0 && type->src_coord != 0){
+      SYMS_USID usid = type->src_coord->usid;
+      SYMS_UnitAccel *unit = syms_group_unit_from_uid(group, usid.uid);
+      SYMS_EnumInfoArray *enum_info_array = syms_push_array(group->arena, SYMS_EnumInfoArray, 1);
+      *enum_info_array = syms_enum_info_array_from_sid(group->arena, group->dbg_data, group->dbg,
+                                                       unit, usid.sid);
+      type->lazy_ptr = enum_info_array;
     }
+    
+    // fill result
+    result = (SYMS_EnumInfoArray*)type->lazy_ptr;
   }
   SYMS_ProfEnd();
   return(result);
 }
 
-SYMS_API SYMS_String8
-syms_type_content_buckets_insert(SYMS_Arena *arena, SYMS_TypeContentBuckets *buckets,
-                                 SYMS_String8 key, SYMS_TypeNode *type){
-  SYMS_ProfBegin("syms_type_content_buckets_insert");
-  SYMS_U64 hash = syms_type_content_hash(key);
-  SYMS_U64 bucket_index = hash%buckets->bucket_count;
-  
-  SYMS_String8 result = syms_push_string_copy(arena, key);
-  
-  SYMS_TypeContentNode *new_node = syms_push_array(arena, SYMS_TypeContentNode, 1);
-  SYMS_StackPush(buckets->buckets[bucket_index], new_node);
-  new_node->key = result;
-  new_node->hash = hash;
-  new_node->type = type;
-  
-  SYMS_ProfEnd();
-  return(result);
-}
-
 ////////////////////////////////
-//~ NOTE(allen): File Mapping
+//~ NOTE(allen): Syms File Mapp
 
 SYMS_API SYMS_Name2FileIDMap*
 syms_group_file_map(SYMS_Group *group){
diff --git a/Engine/Source/Runtime/SymsLib/syms/core/group/syms_group.h b/Engine/Source/Runtime/SymsLib/syms/core/group/syms_group.h
--- a/Engine/Source/Runtime/SymsLib/syms/core/group/syms_group.h
+++ b/Engine/Source/Runtime/SymsLib/syms/core/group/syms_group.h
@@ -4,92 +4,6 @@
 #ifndef SYMS_GROUP_H
 #define SYMS_GROUP_H
 
-////////////////////////////////
-// NOTE(allen): Syms Type Graph
-
-typedef struct SYMS_TypeMember{
-  SYMS_MemKind kind;
-  SYMS_MemVisibility visibility;
-  SYMS_MemFlags flags;
-  SYMS_String8 name;
-  SYMS_U32 off;
-  SYMS_U32 virtual_off;
-  struct SYMS_TypeNode *type;
-} SYMS_TypeMember;
-
-typedef struct SYMS_TypeSrcCoord{
-  SYMS_USID usid;
-  SYMS_FileID file_id;
-  SYMS_U32 line;
-  SYMS_U32 col;
-} SYMS_TypeSrcCoord;
-
-typedef struct SYMS_TypeMemberArray{
-  SYMS_TypeMember *mems;
-  SYMS_U64 count;
-} SYMS_TypeMemberArray;
-
-typedef struct SYMS_TypeNode{
-  // SYMS_TypeNode extends and completes the information from SYMS_TypeInfo.
-  // See SYMS_TypeInfo for more interpretation info.
-  
-  SYMS_TypeKind kind;
-  SYMS_String8 name;
-  SYMS_U64 byte_size;
-  
-  // when non-null contains the source location of the type's definition
-  SYMS_TypeSrcCoord *src_coord;
-  
-  // (in addition to interpretations of SYMS_TypeInfo 'direct_type')
-  //  SYMS_TypeKind_Forward*   -> the concrete type referenced by the forward reference
-  struct SYMS_TypeNode *direct_type;
-  
-  // 'this_type' meaning depends on kind:
-  //  SYMS_TypeKind_MemberPtr  -> the container type of the member pointer
-  //  SYMS_TypeKind_Proc       -> if non-nil this is the type of an implicit 'this' in a C++ method
-  struct SYMS_TypeNode *this_type;
-  
-  union{
-    // kind: SYMS_TypeKind_Modifier
-    SYMS_TypeModifiers mods;
-    
-    // kind: SYMS_TypeKind_Array
-    SYMS_U64 array_count;
-    
-    // kind: SYMS_TypeKind_Proc
-    struct{
-      struct SYMS_TypeNode **params;
-      SYMS_U64 param_count;
-    } proc;
-    
-    // opaque pointer for lazy eval attachments to the type node.
-    void *lazy_ptr;
-  };
-} SYMS_TypeNode;
-
-typedef struct SYMS_TypeUSIDNode{
-  struct SYMS_TypeUSIDNode *next;
-  SYMS_USID key;
-  SYMS_TypeNode *type;
-} SYMS_TypeUSIDNode;
-
-typedef struct SYMS_TypeUSIDBuckets{
-  SYMS_TypeUSIDNode **buckets;
-  SYMS_U64 bucket_count;
-} SYMS_TypeUSIDBuckets;
-
-typedef struct SYMS_TypeContentNode{
-  struct SYMS_TypeContentNode *next;
-  SYMS_String8 key;
-  SYMS_U64 hash;
-  SYMS_TypeNode *type;
-} SYMS_TypeContentNode;
-
-typedef struct SYMS_TypeContentBuckets{
-  SYMS_TypeContentNode **buckets;
-  SYMS_U64 bucket_count;
-} SYMS_TypeContentBuckets;
-
 ////////////////////////////////
 // NOTE(allen): Syms Line Mapping Structures
 
@@ -141,16 +55,18 @@ typedef struct SYMS_FileToLineToAddrLoose{
 
 typedef SYMS_U32 SYMS_GroupUnitCacheFlags;
 enum{
-  SYMS_GroupUnitCacheFlag_HasProcSidArray  = (1 << 0),
-  SYMS_GroupUnitCacheFlag_HasVarSidArray   = (1 << 1),
-  SYMS_GroupUnitCacheFlag_HasTypeSidArray  = (1 << 2),
-  SYMS_GroupUnitCacheFlag_HasFileTable     = (1 << 3),
-  SYMS_GroupUnitCacheFlag_HasInfFileTable  = (1 << 4),
-  SYMS_GroupUnitCacheFlag_HasLineTable     = (1 << 5),
-  SYMS_GroupUnitCacheFlag_HasProcMap       = (1 << 6),
-  SYMS_GroupUnitCacheFlag_HasVarMap        = (1 << 7),
-  SYMS_GroupUnitCacheFlag_HasLineSeqMap    = (1 << 8),
-  SYMS_GroupUnitCacheFlag_HasLineToAddrMap = (1 << 9),
+  SYMS_GroupUnitCacheFlag_HasProcSidArray   = (1 <<  0),
+  SYMS_GroupUnitCacheFlag_HasVarSidArray    = (1 <<  1),
+  SYMS_GroupUnitCacheFlag_HasTlsVarSidArray = (1 <<  2),
+  SYMS_GroupUnitCacheFlag_HasTypeSidArray   = (1 <<  3),
+  SYMS_GroupUnitCacheFlag_HasFileTable      = (1 <<  4),
+  SYMS_GroupUnitCacheFlag_HasInfFileTable   = (1 <<  5),
+  SYMS_GroupUnitCacheFlag_HasLineTable      = (1 <<  6),
+  SYMS_GroupUnitCacheFlag_HasProcMap        = (1 <<  7),
+  SYMS_GroupUnitCacheFlag_HasVarMap         = (1 <<  8),
+  SYMS_GroupUnitCacheFlag_HasLineSeqMap     = (1 <<  9),
+  SYMS_GroupUnitCacheFlag_HasLineToAddrMap  = (1 <<  10),
+  SYMS_GroupUnitCacheFlag_HasTypeNameMap    = (1 <<  11),
 };
 
 typedef struct SYMS_Group{
@@ -184,6 +100,7 @@ typedef struct SYMS_Group{
   SYMS_UnitAccel **units;
   SYMS_SymbolIDArray *proc_sid_arrays;
   SYMS_SymbolIDArray *var_sid_arrays;
+  SYMS_SymbolIDArray *thread_sid_arrays;
   SYMS_SymbolIDArray *type_sid_arrays;
   SYMS_String8Array *file_tables;
   SYMS_String8Array *inferred_file_tables;
@@ -192,12 +109,14 @@ typedef struct SYMS_Group{
   SYMS_SpatialMap1D *unit_var_maps;
   SYMS_SpatialMap1D *line_sequence_maps;
   SYMS_FileToLineToAddrBuckets *file_to_line_to_addr_buckets;
+  SYMS_SymbolNameMap *unit_type_maps;
   
   //- hash tables caches
   SYMS_StringCons string_cons;
   SYMS_FileID2NameMap file_id_2_name_map;
-  SYMS_TypeUSIDBuckets type_usid_buckets;
-  SYMS_TypeContentBuckets type_content_buckets;
+  
+  //- type graph
+  SYMS_TypeGraph type_graph;
   
   //- one-time fills/builds
   SYMS_B8 unit_ranges_is_filled;
@@ -230,25 +149,6 @@ SYMS_READ_ONLY SYMS_GLOBAL SYMS_LineToAddrMap syms_line_to_addr_map_nil = {0};
 SYMS_READ_ONLY SYMS_GLOBAL SYMS_TypeMemberArray syms_type_member_array_nil = {0};
 SYMS_READ_ONLY SYMS_GLOBAL SYMS_EnumInfoArray syms_enum_info_array_nil = {0};
 
-SYMS_READ_ONLY SYMS_GLOBAL SYMS_TypeNode syms_type_graph_nil = {
-  SYMS_TypeKind_Null,     // kind
-  {(SYMS_U8*)"(nil)", 5}, // name
-  0,                      // byte_size
-  0,                      // src_coord
-  &syms_type_graph_nil,   // direct_type
-  &syms_type_graph_nil,   // this_type
-};
-
-
-////////////////////////////////
-// NOTE(allen): File Inference Helper Macros
-
-// TODO(allen): Just here to help people with transition. Eventually we're
-// going to cut the helper macros just have syms_file_inf_infer*
-#define SYMS_GroupInitParams            SYMS_ParseBundle
-#define syms_group_infer_from_file_list syms_file_inf_infer_from_file_list
-#define syms_group_infer_from_file      syms_file_inf_infer_from_file
-
 ////////////////////////////////
 // NOTE(allen): Syms Group Setup Functions
 
@@ -290,6 +190,7 @@ SYMS_API SYMS_UnitNames     syms_group_unit_names_from_uid(SYMS_Arena *arena, SY
 SYMS_API SYMS_UnitAccel*    syms_group_unit_from_uid(SYMS_Group *group, SYMS_UnitID uid);
 SYMS_API SYMS_SymbolIDArray*syms_group_proc_sid_array_from_uid(SYMS_Group *group, SYMS_UnitID uid);
 SYMS_API SYMS_SymbolIDArray*syms_group_var_sid_array_from_uid(SYMS_Group *group, SYMS_UnitID uid);
+SYMS_API SYMS_SymbolIDArray*syms_group_tls_var_sid_array_from_uid(SYMS_Group *group, SYMS_UnitID uid);
 SYMS_API SYMS_SymbolIDArray*syms_group_type_sid_array_from_uid(SYMS_Group *group, SYMS_UnitID uid);
 SYMS_API SYMS_String8Array* syms_group_file_table_from_uid(SYMS_Group *group, SYMS_UnitID uid);
 SYMS_API SYMS_LineParseOut* syms_group_line_parse_from_uid(SYMS_Group *group, SYMS_UnitID uid);
@@ -325,8 +226,6 @@ SYMS_API SYMS_U64Maybe syms_group_foff_from_voff__linear_scan(SYMS_Group *group,
 SYMS_API SYMS_UnitID   syms_group_uid_from_voff__linear_scan(SYMS_Group *group, SYMS_U64 voff);
 SYMS_API SYMS_SymbolID syms_group_proc_sid_from_uid_voff__linear_scan(SYMS_Group *group, SYMS_UnitID uid,
                                                                       SYMS_U64 voff);
-SYMS_API SYMS_SymbolID syms_group_var_sid_from_uid_voff__linear_scan(SYMS_Group *group, SYMS_UnitID uid,
-                                                                     SYMS_U64 voff);
 
 SYMS_API SYMS_Line     syms_group_line_from_uid_voff__linear_scan(SYMS_Group *group, SYMS_UnitID uid,
                                                                   SYMS_U64 voff);
@@ -342,13 +241,18 @@ SYMS_API SYMS_SpatialMap1D* syms_group_unit_map(SYMS_Group *group);
 
 // thread safe (with lanes equipped to group)
 SYMS_API SYMS_SpatialMap1D* syms_group_proc_map_from_uid(SYMS_Group *group, SYMS_UnitID uid);
-SYMS_API SYMS_SpatialMap1D* syms_group_line_sequence_map_from_uid(SYMS_Group *group, SYMS_UnitID uid);
+SYMS_API SYMS_SpatialMap1D* syms_group_line_sequence_map_from_uid(SYMS_Group *group,
+                                                                  SYMS_UnitID uid);
 
-SYMS_API SYMS_SpatialMap1D* syms_group_var_map_from_uid(SYMS_Group *group, SYMS_UnitID uid);
+SYMS_API void                syms_group_fetch_line_to_addr_maps_from_uid(SYMS_Group *group,
+                                                                         SYMS_UnitID uid);
+SYMS_API SYMS_LineToAddrMap*
+syms_group_line_to_addr_map_from_uid_file_id(SYMS_Group *group,
+                                             SYMS_UnitID uid,
+                                             SYMS_FileID file_id);
 
-SYMS_API void                syms_group_fetch_line_to_addr_maps_from_uid(SYMS_Group *group, SYMS_UnitID uid);
-SYMS_API SYMS_LineToAddrMap* syms_group_line_to_addr_map_from_uid_file_id(SYMS_Group *group, SYMS_UnitID uid,
-                                                                          SYMS_FileID file_id);
+SYMS_API SYMS_SymbolNameMap*
+syms_group_type_map_from_uid(SYMS_Group *group, SYMS_UnitID uid);
 
 SYMS_API SYMS_SpatialMap1D* syms_group_stripped_info_map(SYMS_Group *group);
 
@@ -362,8 +266,6 @@ SYMS_API SYMS_U64Maybe syms_group_sec_foff_from_voff__accelerated(SYMS_Group *gr
 SYMS_API SYMS_UnitID   syms_group_uid_from_voff__accelerated(SYMS_Group *group, SYMS_U64 voff);
 SYMS_API SYMS_SymbolID syms_group_proc_sid_from_uid_voff__accelerated(SYMS_Group *group, SYMS_UnitID uid,
                                                                       SYMS_U64 voff);
-SYMS_API SYMS_SymbolID syms_group_var_sid_from_uid_voff__accelerated(SYMS_Group *group, SYMS_UnitID uid,
-                                                                     SYMS_U64 voff);
 
 SYMS_API SYMS_Line     syms_group_line_from_uid_voff__accelerated(SYMS_Group *group, SYMS_UnitID uid,
                                                                   SYMS_U64 voff);
@@ -378,48 +280,56 @@ SYMS_API void          syms_line_to_addr_line_sort(SYMS_FileToLineToAddrLooseLin
 SYMS_API void          syms_line_to_addr_line_sort__rec(SYMS_FileToLineToAddrLooseLine **array, SYMS_U64 count);
 
 ////////////////////////////////
-// NOTE(allen): Syms Group Type Graph Functions
-
-SYMS_API SYMS_TypeNode* syms_group_type_from_usid(SYMS_Group *group, SYMS_USID usid);
-SYMS_API SYMS_TypeNode* syms_group_type_from_usid__rec(SYMS_Group *group, SYMS_USID usid);
-
-SYMS_API SYMS_TypeMemberArray* syms_group_members_from_type(SYMS_Group *group, SYMS_TypeNode *type);
-SYMS_API SYMS_EnumInfoArray*   syms_group_enum_members_from_type(SYMS_Group *group, SYMS_TypeNode *type);
-
-SYMS_API SYMS_U64       syms_group_type_size_from_usid(SYMS_Group *group, SYMS_USID usid);
-
-SYMS_API SYMS_TypeNode* syms_group_type_basic(SYMS_Group *group, SYMS_TypeKind basic_kind,
-                                              SYMS_U64 size, SYMS_String8 name);
-SYMS_API SYMS_TypeNode* syms_group_type_mod_from_type(SYMS_Group *group, SYMS_TypeNode *type,
-                                                      SYMS_TypeModifiers mods);
-SYMS_API SYMS_TypeNode* syms_group_type_ptr_from_type(SYMS_Group *group, SYMS_TypeKind ptr_kind,
-                                                      SYMS_TypeNode *type);
-SYMS_API SYMS_TypeNode* syms_group_type_array_from_type(SYMS_Group *group, SYMS_TypeNode *type, SYMS_U64 count);
-SYMS_API SYMS_TypeNode* syms_group_type_proc_from_type(SYMS_Group *group,
-                                                       SYMS_TypeNode *ret_type, SYMS_TypeNode *this_type,
-                                                       SYMS_TypeNode **param_types, SYMS_U64 param_count);
-SYMS_API SYMS_TypeNode* syms_group_type_member_ptr_from_type(SYMS_Group *group, SYMS_TypeNode *container,
-                                                             SYMS_TypeNode *type);
-
-SYMS_API SYMS_String8   syms_group_string_from_type(SYMS_Arena *arena, SYMS_TypeNode *type);
-
-SYMS_API void           syms_group_lhs_string_from_type(SYMS_Arena *arena, SYMS_TypeNode *type,
-                                                        SYMS_String8List *out);
-SYMS_API void           syms_group_rhs_string_from_type(SYMS_Arena *arena, SYMS_TypeNode *type,
-                                                        SYMS_String8List *out);
-
-SYMS_API void           syms_group_lhs_string_from_type_skip_return(SYMS_Arena *arena, SYMS_TypeNode *type,
-                                                                    SYMS_String8List *out);
-
-SYMS_API SYMS_U64       syms_type_usid_hash(SYMS_USID usid);
-SYMS_API SYMS_TypeNode* syms_type_from_usid(SYMS_TypeUSIDBuckets *buckets, SYMS_USID usid);
-SYMS_API void           syms_type_usid_buckets_insert(SYMS_Arena *arena, SYMS_TypeUSIDBuckets *buckets,
-                                                      SYMS_USID key, SYMS_TypeNode *type);
-
-SYMS_API SYMS_U64       syms_type_content_hash(SYMS_String8 data);
-SYMS_API SYMS_TypeNode* syms_type_from_content(SYMS_TypeContentBuckets *buckets, SYMS_String8 data);
-SYMS_API SYMS_String8   syms_type_content_buckets_insert(SYMS_Arena *arena, SYMS_TypeContentBuckets *buckets,
-                                                         SYMS_String8 key, SYMS_TypeNode *type);
+// NOTE(allen): Syms Group Type Graph
+
+SYMS_API SYMS_TypeGraph* syms_group_type_graph(SYMS_Group *group);
+
+////////////////////////////////
+// NOTE(allen): Syms Group Varaible Address Mapping Functions
+
+// these are seperated because they require type info which
+// cannot be parallelized like the others
+
+
+SYMS_API SYMS_SymbolID syms_group_var_sid_from_uid_voff__linear_scan(SYMS_TypeGraph *graph,
+                                                                     SYMS_Group *group,
+                                                                     SYMS_UnitID uid,
+                                                                     SYMS_U64 voff);
+
+SYMS_API SYMS_SpatialMap1D* syms_group_var_map_from_uid(SYMS_TypeGraph *graph,
+                                                        SYMS_Group *group, SYMS_UnitID uid);
+
+SYMS_API SYMS_SymbolID syms_group_var_sid_from_uid_voff__accelerated(SYMS_TypeGraph *graph,
+                                                                     SYMS_Group *group,
+                                                                     SYMS_UnitID uid,
+                                                                     SYMS_U64 voff);
+
+
+////////////////////////////////
+// NOTE(allen): Syms Group Type Info Functions
+
+SYMS_API SYMS_TypeNode* syms_group_type_from_usid(SYMS_TypeGraph *graph, SYMS_Group *group,
+                                                  SYMS_USID usid);
+SYMS_API SYMS_TypeNode* syms_group_type_from_usid__rec(SYMS_TypeGraph *graph,
+                                                       SYMS_Group *group, SYMS_USID usid);
+
+SYMS_API SYMS_TypeChain syms_group_artificial_types_from_name(SYMS_Group *group,
+                                                              SYMS_String8 name);
+
+SYMS_API SYMS_USIDList syms_group_type_list_from_name_accelerated(SYMS_Arena *arena,
+                                                                  SYMS_Group *group,
+                                                                  SYMS_String8 name);
+
+SYMS_API SYMS_SymbolIDArray syms_group_types_from_unit_name(SYMS_Group *group,
+                                                            SYMS_UnitID uid, SYMS_String8 name);
+
+SYMS_API SYMS_U64       syms_group_type_size_from_usid(SYMS_TypeGraph *graph,
+                                                       SYMS_Group *group, SYMS_USID usid);
+
+SYMS_API SYMS_TypeMemberArray* syms_type_members_from_type(SYMS_Group *group,
+                                                           SYMS_TypeNode *type);
+SYMS_API SYMS_EnumInfoArray*   syms_type_enum_members_from_type(SYMS_Group *group,
+                                                                SYMS_TypeNode *type);
 
 ////////////////////////////////
 //~ NOTE(allen): Syms File Map
diff --git a/Engine/Source/Runtime/SymsLib/syms/core/group/syms_type_graph.c b/Engine/Source/Runtime/SymsLib/syms/core/group/syms_type_graph.c
new file mode 100644
--- /dev/null
+++ b/Engine/Source/Runtime/SymsLib/syms/core/group/syms_type_graph.c
@@ -0,0 +1,784 @@
+// Copyright Epic Games, Inc. All Rights Reserved.
+
+#ifndef SYMS_TYPE_GRAPH_C
+#define SYMS_TYPE_GRAPH_C
+
+////////////////////////////////
+// NOTE(allen): Type Graph Setup Functions
+
+SYMS_API void
+syms_type_graph_init(SYMS_TypeGraph *graph,
+                     SYMS_Arena *graph_arena,
+                     SYMS_StringCons *graph_string_cons,
+                     SYMS_U64 address_size){
+  SYMS_ProfBegin("syms_type_graph_init");
+  
+  graph->arena = graph_arena;
+  graph->string_cons = graph_string_cons;
+  graph->address_size = address_size;
+  
+  graph->type_usid_buckets.bucket_count = 4093;
+  graph->type_usid_buckets.buckets = syms_push_array_zero(graph->arena, SYMS_TypeUSIDNode*, 4093);
+  
+  graph->type_content_buckets.bucket_count = 4093;
+  graph->type_content_buckets.buckets = syms_push_array_zero(graph->arena, SYMS_TypeContentNode*, 4093);
+  
+  graph->type_name_buckets.bucket_count = 4093;
+  graph->type_name_buckets.buckets = syms_push_array_zero(graph->arena, SYMS_TypeNameNode*, 4093);
+  
+  SYMS_ProfEnd();
+}
+
+////////////////////////////////
+// NOTE(allen): Type Content Hash Functions
+
+SYMS_API SYMS_U64
+syms_type_content_hash(SYMS_String8 data){
+  SYMS_U64 result = syms_hash_djb2(data);
+  return(result);
+}
+
+SYMS_API SYMS_TypeNode*
+syms_type_from_content(SYMS_TypeContentBuckets *buckets, SYMS_String8 data){
+  SYMS_ProfBegin("syms_type_from_content");
+  SYMS_TypeNode *result = 0;
+  if (buckets->bucket_count != 0){
+    SYMS_U64 hash = syms_type_content_hash(data);
+    SYMS_U64 bucket_index = hash%buckets->bucket_count;
+    for (SYMS_TypeContentNode *node = buckets->buckets[bucket_index];
+         node != 0;
+         node = node->next){
+      if (node->hash == hash && syms_string_match(node->key, data, 0)){
+        result = node->type;
+        break;
+      }
+    }
+  }
+  SYMS_ProfEnd();
+  return(result);
+}
+
+SYMS_API SYMS_String8
+syms_type_content_buckets_insert(SYMS_Arena *arena, SYMS_TypeContentBuckets *buckets,
+                                 SYMS_String8 key, SYMS_TypeNode *type){
+  SYMS_ProfBegin("syms_type_content_buckets_insert");
+  SYMS_U64 hash = syms_type_content_hash(key);
+  SYMS_String8 result = {0};
+  if (buckets->bucket_count > 0){
+    SYMS_U64 bucket_index = hash%buckets->bucket_count;
+    
+    result = syms_push_string_copy(arena, key);
+    
+    SYMS_TypeContentNode *new_node = syms_push_array(arena, SYMS_TypeContentNode, 1);
+    SYMS_StackPush(buckets->buckets[bucket_index], new_node);
+    new_node->key = result;
+    new_node->hash = hash;
+    new_node->type = type;
+  }
+  SYMS_ProfEnd();
+  return(result);
+}
+
+////////////////////////////////
+// NOTE(allen): Type USID Hash Functions
+
+SYMS_API SYMS_U64
+syms_type_usid_hash(SYMS_USID usid){
+  SYMS_U64 result = syms_hash_u64(usid.sid + usid.uid*97);
+  return(result);
+}
+
+SYMS_API SYMS_TypeNode*
+syms_type_from_usid(SYMS_TypeUSIDBuckets *buckets, SYMS_USID usid){
+  SYMS_ProfBegin("syms_type_from_usid");
+  SYMS_TypeNode *result = 0;
+  if (buckets->bucket_count > 0){
+    SYMS_U64 hash = syms_type_usid_hash(usid);
+    SYMS_U64 bucket_index = hash%buckets->bucket_count;
+    for (SYMS_TypeUSIDNode *node = buckets->buckets[bucket_index];
+         node != 0;
+         node = node->next){
+      SYMS_USID key = node->key;
+      if (key.uid == usid.uid && key.sid == usid.sid){
+        result = node->type;
+        break;
+      }
+    }
+  }
+  SYMS_ProfEnd();
+  return(result);
+}
+
+SYMS_API void
+syms_type_usid_buckets_insert(SYMS_Arena *arena, SYMS_TypeUSIDBuckets *buckets,
+                              SYMS_USID key,
+                              SYMS_TypeNode *type){
+  SYMS_ProfBegin("syms_type_usid_buckets_insert");
+  if (buckets->bucket_count > 0){
+    SYMS_U64 hash = syms_type_usid_hash(key);
+    SYMS_U64 bucket_index = hash%buckets->bucket_count;
+    SYMS_TypeUSIDNode *new_node = syms_push_array(arena, SYMS_TypeUSIDNode, 1);
+    SYMS_StackPush(buckets->buckets[bucket_index], new_node);
+    new_node->key = key;
+    new_node->type = type;
+  }
+  SYMS_ProfEnd();
+}
+
+////////////////////////////////
+// NOTE(allen): Type Name Hash Functions
+
+SYMS_API SYMS_U64
+syms_type_name_hash(SYMS_U8 *ptr){
+  SYMS_U64 u64 = (SYMS_U64)(ptr);
+  SYMS_U64 bits = u64 >> 3;
+  SYMS_U64 result = bits ^ (bits << 32);
+  return(result);
+}
+
+SYMS_API SYMS_TypeChain*
+syms_type_chain_from_name(SYMS_TypeNameBuckets *buckets, SYMS_U8 *name_ptr){
+  SYMS_TypeChain *result = 0;
+  SYMS_ProfBegin("syms_type_chain_from_name");
+  if (buckets->bucket_count > 0){
+    SYMS_U64 hash = syms_type_name_hash(name_ptr);
+    SYMS_U64 bucket_index = hash%buckets->bucket_count;
+    for (SYMS_TypeNameNode *node = buckets->buckets[bucket_index];
+         node != 0;
+         node = node->next){
+      if (node->name_ptr == name_ptr){
+        result = &node->chain;
+        break;
+      }
+    }
+  }
+  SYMS_ProfEnd();
+  return(result);
+}
+
+SYMS_API void
+syms_type_name_buckets_insert(SYMS_Arena *arena, SYMS_TypeNameBuckets *buckets,
+                              SYMS_U8 *name_ptr, SYMS_TypeNode *type){
+  SYMS_ProfBegin("syms_type_chain_from_name");
+  if (buckets->bucket_count > 0){
+    // find existing name node
+    SYMS_TypeNameNode *match = 0;
+    SYMS_U64 hash = syms_type_name_hash(name_ptr);
+    SYMS_U64 bucket_index = hash%buckets->bucket_count;
+    for (SYMS_TypeNameNode *node = buckets->buckets[bucket_index];
+         node != 0;
+         node = node->next){
+      if (node->name_ptr == name_ptr){
+        match = node;
+        break;
+      }
+    }
+    
+    // new name node if necessary
+    if (match == 0){
+      SYMS_TypeNameNode *new_node = syms_push_array_zero(arena, SYMS_TypeNameNode, 1);
+      SYMS_StackPush(buckets->buckets[bucket_index], new_node);
+      new_node->name_ptr = name_ptr;
+      
+      match = new_node;
+    }
+    
+    // new type chain node
+    {
+      SYMS_TypeChainNode *new_chain = syms_push_array(arena, SYMS_TypeChainNode, 1);
+      SYMS_QueuePush(match->chain.first, match->chain.last, new_chain);
+      match->chain.count += 1;
+      new_chain->type = type;
+    }
+  }
+  SYMS_ProfEnd();
+}
+
+////////////////////////////////
+// NOTE(allen): Type Name Lookup Functions
+
+SYMS_API SYMS_TypeChain
+syms_type_from_name(SYMS_TypeGraph *graph, SYMS_String8 name){
+  SYMS_TypeChain result = {0};
+  SYMS_String8 name_cons = syms_string_cons(graph->arena, graph->string_cons, name);
+  SYMS_TypeChain *chain = syms_type_chain_from_name(&graph->type_name_buckets, name_cons.str);
+  if (chain != 0){
+    syms_memmove(&result, chain, sizeof(*chain));
+  }
+  return(result);
+}
+
+////////////////////////////////
+// NOTE(allen): Type Node Constructors
+
+SYMS_API SYMS_TypeNode*
+syms_type_basic(SYMS_TypeGraph *graph, SYMS_TypeKind basic_kind,
+                SYMS_U64 size, SYMS_String8 name){
+  SYMS_ProfBegin("syms_type_basic");
+  
+  SYMS_TypeNode *result = &syms_type_graph_nil;
+  
+  SYMS_String8 name_cons = syms_string_cons(graph->arena, graph->string_cons, name);
+  
+  SYMS_U64 content[3];
+  content[0] = basic_kind;
+  content[1] = size;
+  content[2] = (SYMS_U64)name_cons.str;
+  SYMS_String8 data = {(SYMS_U8*)content, sizeof(content)};
+  result = syms_type_from_content(&graph->type_content_buckets, data);
+  
+  if (result == 0){
+    result = syms_push_array_zero(graph->arena, SYMS_TypeNode, 1);
+    syms_type_content_buckets_insert(graph->arena, &graph->type_content_buckets, data, result);
+    syms_type_name_buckets_insert(graph->arena, &graph->type_name_buckets,
+                                  name_cons.str, result);
+    
+    result->kind = basic_kind;
+    result->name = name_cons;
+    result->byte_size = size;
+    result->direct_type = &syms_type_graph_nil;
+    result->this_type = &syms_type_graph_nil;
+  }
+  
+  SYMS_ProfEnd();
+  return(result);
+}
+
+SYMS_API SYMS_TypeNode*
+syms_type_mod_from_type(SYMS_TypeGraph *graph, SYMS_TypeNode *type, SYMS_TypeModifiers mods){
+  SYMS_ProfBegin("syms_type_mod_from_type");
+  
+  SYMS_TypeNode *result = &syms_type_graph_nil;
+  
+  SYMS_U64 content[3];
+  content[0] = SYMS_TypeKind_Modifier;
+  content[1] = mods;
+  content[2] = (SYMS_U64)type;
+  SYMS_String8 data = {(SYMS_U8*)content, sizeof(content)};
+  result = syms_type_from_content(&graph->type_content_buckets, data);
+  
+  if (result == 0){
+    result = syms_push_array_zero(graph->arena, SYMS_TypeNode, 1);
+    syms_type_content_buckets_insert(graph->arena, &graph->type_content_buckets, data, result);
+    
+    result->kind = SYMS_TypeKind_Modifier;
+    result->byte_size = type->byte_size;
+    result->direct_type = type;
+    result->this_type = &syms_type_graph_nil;
+    result->mods = mods;
+  }
+  
+  SYMS_ProfEnd();
+  return(result);
+}
+
+SYMS_API SYMS_TypeNode*
+syms_type_ptr_from_type(SYMS_TypeGraph *graph, SYMS_TypeKind ptr_kind, SYMS_TypeNode *type){
+  SYMS_ProfBegin("syms_type_ptr_from_type");
+  
+  SYMS_TypeNode *result = &syms_type_graph_nil;
+  
+  SYMS_U64 content[2];
+  content[0] = ptr_kind;
+  content[1] = (SYMS_U64)type;
+  SYMS_String8 data = {(SYMS_U8*)content, sizeof(content)};
+  result = syms_type_from_content(&graph->type_content_buckets, data);
+  
+  if (result == 0){
+    result = syms_push_array_zero(graph->arena, SYMS_TypeNode, 1);
+    syms_type_content_buckets_insert(graph->arena, &graph->type_content_buckets, data, result);
+    
+    result->kind = ptr_kind;
+    result->byte_size = graph->address_size;
+    result->direct_type = type;
+    result->this_type = &syms_type_graph_nil;
+  }
+  
+  SYMS_ProfEnd();
+  return(result);
+}
+
+SYMS_API SYMS_TypeNode*
+syms_type_array_from_type(SYMS_TypeGraph *graph, SYMS_TypeNode *type, SYMS_U64 count){
+  SYMS_ProfBegin("syms_type_array_from_type");
+  
+  SYMS_TypeNode *result = &syms_type_graph_nil;
+  
+  SYMS_U64 content[3];
+  content[0] = SYMS_TypeKind_Array;
+  content[1] = count;
+  content[2] = (SYMS_U64)type;
+  SYMS_String8 data = {(SYMS_U8*)content, sizeof(content)};
+  result = syms_type_from_content(&graph->type_content_buckets, data);
+  
+  if (result == 0){
+    result = syms_push_array_zero(graph->arena, SYMS_TypeNode, 1);
+    syms_type_content_buckets_insert(graph->arena, &graph->type_content_buckets, data, result);
+    
+    result->kind = SYMS_TypeKind_Array;
+    result->byte_size = count*type->byte_size;
+    result->direct_type = type;
+    result->this_type = &syms_type_graph_nil;
+    result->array_count = count;
+  }
+  
+  SYMS_ProfEnd();
+  return(result);
+}
+
+SYMS_API SYMS_TypeNode*
+syms_type_proc_from_type(SYMS_TypeGraph *graph, SYMS_TypeNode *ret_type, SYMS_TypeNode *this_type,
+                         SYMS_TypeNode **param_types, SYMS_U64 param_count){
+  SYMS_ProfBegin("syms_type_proc_from_type");
+  
+  SYMS_TypeNode *result = &syms_type_graph_nil;
+  
+  SYMS_String8 params_as_string = {(SYMS_U8*)param_types, sizeof(*param_types)*param_count};
+  SYMS_String8 params_cons = syms_string_cons(graph->arena, graph->string_cons, params_as_string);
+  
+  SYMS_U64 content[4];
+  content[0] = SYMS_TypeKind_Proc;
+  content[1] = (SYMS_U64)ret_type;
+  content[2] = (SYMS_U64)this_type;
+  content[3] = (SYMS_U64)params_cons.str;
+  SYMS_String8 data = {(SYMS_U8*)content, sizeof(content)};
+  result = syms_type_from_content(&graph->type_content_buckets, data);
+  
+  if (result == 0){
+    result = syms_push_array_zero(graph->arena, SYMS_TypeNode, 1);
+    syms_type_content_buckets_insert(graph->arena, &graph->type_content_buckets, data, result);
+    
+    result->kind = SYMS_TypeKind_Proc;
+    result->byte_size = 0;
+    result->direct_type = ret_type;
+    result->this_type = this_type;
+    result->proc.params = (SYMS_TypeNode**)(params_cons.str);
+    result->proc.param_count = param_count;
+  }
+  
+  SYMS_ProfEnd();
+  return(result);
+}
+
+SYMS_API SYMS_TypeNode*
+syms_type_member_ptr_from_type(SYMS_TypeGraph *graph, SYMS_TypeNode *container, SYMS_TypeNode *type){
+  SYMS_ProfBegin("syms_type_member_ptr_from_type");
+  
+  SYMS_TypeNode *result = &syms_type_graph_nil;
+  
+  SYMS_U64 content[3];
+  content[0] = SYMS_TypeKind_MemberPtr;
+  content[1] = (SYMS_U64)container;
+  content[2] = (SYMS_U64)type;
+  SYMS_String8 data = {(SYMS_U8*)content, sizeof(content)};
+  result = syms_type_from_content(&graph->type_content_buckets, data);
+  
+  if (result == 0){
+    result = syms_push_array_zero(graph->arena, SYMS_TypeNode, 1);
+    syms_type_content_buckets_insert(graph->arena, &graph->type_content_buckets, data, result);
+    
+    SYMS_U64 byte_size = graph->address_size;
+    if (type->kind == SYMS_TypeKind_Proc){
+      byte_size *= 2;
+    }
+    
+    result->kind = SYMS_TypeKind_MemberPtr;
+    result->byte_size = byte_size;
+    result->direct_type = type;
+    result->this_type = container;
+  }
+  
+  SYMS_ProfEnd();
+  return(result);
+}
+
+////////////////////////////////
+// NOTE(allen): Type Info Operators
+
+SYMS_API SYMS_TypeNode*
+syms_type_resolved(SYMS_TypeNode *type){
+  SYMS_TypeNode *result = type;
+  for (;syms_type_kind_is_forward(result->kind) ||
+       result->kind == SYMS_TypeKind_Modifier ||
+       result->kind == SYMS_TypeKind_Typedef;){
+    result = result->direct_type;
+  }
+  return(result);
+}
+
+SYMS_API SYMS_B32
+syms_type_node_match(SYMS_TypeNode *l, SYMS_TypeNode *r){
+  SYMS_B32 result = syms_false;
+  
+  // resolve forward references
+  SYMS_TypeNode *lres = syms_type_resolved(l);
+  SYMS_TypeNode *rres = syms_type_resolved(r);
+  
+  if (lres == rres){
+    result = syms_true;
+  }
+  else{
+    if (lres->kind == rres->kind){
+      switch (lres->kind){
+        default:
+        {
+          result = syms_true;
+        }break;
+        
+        case SYMS_TypeKind_Struct:
+        case SYMS_TypeKind_Union:
+        case SYMS_TypeKind_Enum:
+        case SYMS_TypeKind_Typedef:
+        {
+          if (lres->src_coord != 0 && rres->src_coord != 0){
+            SYMS_USID lusid = lres->src_coord->usid;
+            SYMS_USID rusid = rres->src_coord->usid;
+            if (lusid.uid == rusid.uid && lusid.sid == rusid.sid){
+              result = syms_true;
+            }
+          }
+        }break;
+        
+        case SYMS_TypeKind_Ptr:
+        case SYMS_TypeKind_LValueReference:
+        case SYMS_TypeKind_RValueReference:
+        {
+          if (syms_type_node_match(lres->direct_type, rres->direct_type)){
+            result = syms_true;
+          }
+        }break;
+        
+        case SYMS_TypeKind_MemberPtr:
+        {
+          if (syms_type_node_match(lres->direct_type, rres->direct_type) &&
+              syms_type_node_match(lres->this_type, rres->this_type)){
+            result = syms_true;
+          }
+        }break;
+        
+        case SYMS_TypeKind_Array:
+        {
+          if (lres->array_count == rres->array_count &&
+              syms_type_node_match(lres->direct_type, rres->direct_type)){
+            result = syms_true;
+          }
+        }break;
+        
+        case SYMS_TypeKind_Proc:
+        {
+          if (lres->proc.param_count == rres->proc.param_count &&
+              syms_type_node_match(lres->direct_type, rres->direct_type) &&
+              syms_type_node_match(lres->this_type, rres->this_type)){
+            SYMS_B32 params_match = syms_true;
+            SYMS_TypeNode **lp = lres->proc.params;
+            SYMS_TypeNode **rp = rres->proc.params;
+            SYMS_U64 count = lres->proc.param_count;
+            for (SYMS_U64 i = 0; i < count; i += 1, lp += 1, rp += 1){
+              if (!syms_type_node_match(*lp, *rp)){
+                params_match = syms_false;
+                break;
+              }
+            }
+            result = params_match;
+          }
+        }break;
+      }
+    }
+  }
+  return(result);
+}
+
+SYMS_API SYMS_TypeNode*
+syms_type_resolve_enum_to_basic(SYMS_TypeGraph *graph, SYMS_TypeNode *t){
+  SYMS_TypeNode *result = t;
+  if (t->kind == SYMS_TypeKind_Enum){
+    result = t->direct_type;
+  }
+  return(result);
+}
+
+SYMS_API SYMS_TypeNode*
+syms_type_promoted_from_type_node(SYMS_TypeGraph *graph, SYMS_TypeNode *c){
+  SYMS_TypeNode *result = c;
+  SYMS_TypeKind kind = c->kind;
+  if (kind == SYMS_TypeKind_Bool ||
+      kind == SYMS_TypeKind_Int8  ||
+      kind == SYMS_TypeKind_Int16 ||
+      kind == SYMS_TypeKind_UInt8 ||
+      kind == SYMS_TypeKind_UInt16){
+    result = syms_type_basic(graph, SYMS_TypeKind_Int32, 4, syms_str8_lit("int32"));
+  }
+  return(result);
+}
+
+SYMS_API SYMS_TypeNode*
+syms_type_auto_casted_from_type_nodes(SYMS_TypeGraph *graph, SYMS_TypeNode *l, SYMS_TypeNode *r){
+  SYMS_ASSERT(syms_type_kind_is_basic_or_enum(l->kind) &&
+              syms_type_kind_is_basic_or_enum(r->kind));
+  
+  // replace enums with corresponding ints
+  SYMS_TypeNode *lt = syms_type_resolve_enum_to_basic(graph, l);
+  SYMS_TypeNode *rt = syms_type_resolve_enum_to_basic(graph, r);
+  
+  // first promote each
+  SYMS_TypeNode *lp = syms_type_promoted_from_type_node(graph, lt);
+  SYMS_TypeNode *rp = syms_type_promoted_from_type_node(graph, rt);
+  SYMS_TypeKind lk = lp->kind;
+  SYMS_TypeKind rk = rp->kind;
+  
+  // pick the type with higher rank
+  static SYMS_TypeKind type_table[] = {
+    SYMS_TypeKind_Float64,
+    SYMS_TypeKind_Float32,
+    SYMS_TypeKind_UInt64,
+    SYMS_TypeKind_Int64,
+    SYMS_TypeKind_UInt32,
+  };
+  SYMS_TypeKind k = SYMS_TypeKind_Int32;
+  for (SYMS_U64 i = 0; i < SYMS_ARRAY_SIZE(type_table); i += 1){
+    SYMS_TypeKind check = type_table[i];
+    if (check == lk || check == rk){
+      k = check;
+      break;
+    }
+  }
+  
+  // construct the output type
+  SYMS_TypeNode *result = &syms_type_graph_nil;
+  if (k == lk){
+    result = l;
+  }
+  else if (k == rk){
+    result = r;
+  }
+  else{
+    SYMS_String8 name = syms_string_from_enum_value(SYMS_TypeKind, k);
+    SYMS_U64 bit_size = syms_bit_size_from_type_kind(k);
+    SYMS_U64 byte_size = bit_size >> 3;
+    result = syms_type_basic(graph, k, byte_size, name);
+  }
+  
+  return(result);
+}
+
+////////////////////////////////
+// NOTE(allen): Type Stringizing
+
+SYMS_API SYMS_String8
+syms_string_from_type(SYMS_Arena *arena, SYMS_TypeNode *type){
+  SYMS_ProfBegin("syms_string_from_type");
+  SYMS_ArenaTemp scratch = syms_get_scratch(&arena, 1);
+  SYMS_String8List list = {0};
+  syms_lhs_string_from_type(scratch.arena, type, &list);
+  syms_rhs_string_from_type(scratch.arena, type, &list);
+  SYMS_String8 result = syms_string_list_join(arena, &list, 0);
+  syms_release_scratch(scratch);
+  SYMS_ProfEnd();
+  return(result);
+}
+
+SYMS_API void
+syms_lhs_string_from_type(SYMS_Arena *arena, SYMS_TypeNode *type, SYMS_String8List *out){
+  SYMS_ProfBegin("syms_lhs_string_from_type");
+  syms_lhs_string_from_type__internal(arena, type, out, 0, syms_false);
+  SYMS_ProfEnd();
+}
+
+SYMS_API void
+syms_rhs_string_from_type(SYMS_Arena *arena, SYMS_TypeNode *type, SYMS_String8List *out){
+  SYMS_ProfBegin("syms_rhs_string_from_type");
+  syms_rhs_string_from_type__internal(arena, type, out, 0);
+  SYMS_ProfEnd();
+}
+
+SYMS_API void
+syms_lhs_string_from_type_skip_return(SYMS_Arena *arena, SYMS_TypeNode *type, SYMS_String8List *out){
+  SYMS_ProfBegin("syms_lhs_string_from_type_skip_return");
+  syms_lhs_string_from_type__internal(arena, type, out, 0, syms_true);
+  SYMS_ProfEnd();
+}
+
+SYMS_API void
+syms_lhs_string_from_type__internal(SYMS_Arena *arena, SYMS_TypeNode *type,
+                                    SYMS_String8List *out,
+                                    SYMS_U32 precedence_level, SYMS_B32 skip_return){
+  switch (type->kind){
+    default:
+    {
+      syms_string_list_push(arena, out, type->name);
+      syms_string_list_push(arena, out, syms_str8_lit(" "));
+    }break;
+    
+    case SYMS_TypeKind_Bitfield:
+    {
+      syms_lhs_string_from_type__internal(arena, type->direct_type, out, precedence_level, skip_return);
+    }break;
+    
+    case SYMS_TypeKind_Modifier:
+    {
+      syms_lhs_string_from_type__internal(arena, type->direct_type, out, 1, skip_return);
+      if (type->mods & SYMS_TypeModifier_Const){
+        syms_string_list_push(arena, out, syms_str8_lit("const "));
+      }
+      if (type->mods & SYMS_TypeModifier_Restrict){
+        syms_string_list_push(arena, out, syms_str8_lit("restrict "));
+      }
+      if (type->mods & SYMS_TypeModifier_Volatile){
+        syms_string_list_push(arena, out, syms_str8_lit("volatile "));
+      }
+    }break;
+    
+    case SYMS_TypeKind_Variadic:
+    {
+      syms_string_list_push(arena, out, syms_str8_lit("..."));
+    }break;
+    
+    case SYMS_TypeKind_Label:
+    {
+      // TODO(allen): ???
+    }break;
+    
+    case SYMS_TypeKind_Struct:
+    case SYMS_TypeKind_Union:
+    case SYMS_TypeKind_Enum:
+    case SYMS_TypeKind_Class:
+    case SYMS_TypeKind_Typedef:
+    {
+      syms_string_list_push(arena, out, type->name);
+      syms_string_list_push(arena, out, syms_str8_lit(" "));
+    }break;
+    
+    case SYMS_TypeKind_ForwardStruct:
+    case SYMS_TypeKind_ForwardUnion:
+    case SYMS_TypeKind_ForwardEnum:
+    case SYMS_TypeKind_ForwardClass:
+    {
+      SYMS_String8 keyword = syms_string_from_enum_value(SYMS_TypeKind, type->kind);
+      syms_string_list_push(arena, out, keyword);
+      syms_string_list_push(arena, out, syms_str8_lit(" "));
+      syms_string_list_push(arena, out, type->name);
+      syms_string_list_push(arena, out, syms_str8_lit(" "));
+    }break;
+    
+    case SYMS_TypeKind_Array:
+    {
+      syms_lhs_string_from_type__internal(arena, type->direct_type, out, 2, skip_return);
+      if (precedence_level == 1){
+        syms_string_list_push(arena, out, syms_str8_lit("("));
+      }
+    }break;
+    
+    case SYMS_TypeKind_Proc:
+    {
+      if (!skip_return){
+        syms_lhs_string_from_type__internal(arena, type->direct_type, out, 2, syms_false);
+      }
+      if (precedence_level == 1){
+        syms_string_list_push(arena, out, syms_str8_lit("("));
+      }
+    }break;
+    
+    case SYMS_TypeKind_Ptr:
+    {
+      syms_lhs_string_from_type__internal(arena, type->direct_type, out, 1, skip_return);
+      syms_string_list_push(arena, out, syms_str8_lit("*"));
+    }break;
+    
+    case SYMS_TypeKind_LValueReference:
+    {
+      syms_lhs_string_from_type__internal(arena, type->direct_type, out, 1, skip_return);
+      syms_string_list_push(arena, out, syms_str8_lit("&"));
+    }break;
+    
+    case SYMS_TypeKind_RValueReference:
+    {
+      syms_lhs_string_from_type__internal(arena, type->direct_type, out, 1, skip_return);
+      syms_string_list_push(arena, out, syms_str8_lit("&&"));
+    }break;
+    
+    case SYMS_TypeKind_MemberPtr:
+    {
+      syms_lhs_string_from_type__internal(arena, type->direct_type, out, 1, skip_return);
+      SYMS_TypeNode *container = type->this_type;
+      if (container != 0){
+        syms_string_list_push(arena, out, container->name);
+      }
+      else{
+        syms_string_list_push(arena, out, syms_str8_lit("<unknown-class>"));
+      }
+      syms_string_list_push(arena, out, syms_str8_lit("::*"));
+    }break;
+  }
+}
+
+SYMS_API void
+syms_rhs_string_from_type__internal(SYMS_Arena *arena, SYMS_TypeNode *type,
+                                    SYMS_String8List *out,
+                                    SYMS_U32 precedence_level){
+  switch (type->kind){
+    default:break;
+    
+    case SYMS_TypeKind_Bitfield:
+    {
+      syms_rhs_string_from_type__internal(arena, type->direct_type, out, precedence_level);
+    }break;
+    
+    case SYMS_TypeKind_Modifier:
+    case SYMS_TypeKind_Ptr:
+    case SYMS_TypeKind_LValueReference:
+    case SYMS_TypeKind_RValueReference:
+    case SYMS_TypeKind_MemberPtr:
+    {
+      syms_rhs_string_from_type__internal(arena, type->direct_type, out, 1);
+    }break;
+    
+    case SYMS_TypeKind_Array:
+    {
+      if (precedence_level == 1){
+        syms_string_list_push(arena, out, syms_str8_lit(")"));
+      }
+      
+      SYMS_String8 count_str = syms_string_from_u64(arena, type->array_count);
+      syms_string_list_push(arena, out, syms_str8_lit("["));
+      syms_string_list_push(arena, out, count_str);
+      syms_string_list_push(arena, out, syms_str8_lit("]"));
+      
+      syms_rhs_string_from_type__internal(arena, type->direct_type, out, 2);
+    }break;
+    
+    case SYMS_TypeKind_Proc:
+    {
+      if (precedence_level == 1){
+        syms_string_list_push(arena, out, syms_str8_lit(")"));
+      }
+      
+      // parameters
+      if (type->proc.param_count == 0){
+        syms_string_list_push(arena, out, syms_str8_lit("(void)"));
+      }
+      else{
+        syms_string_list_push(arena, out, syms_str8_lit("("));
+        SYMS_U64 param_count = type->proc.param_count;
+        SYMS_TypeNode **param = type->proc.params;
+        for (SYMS_U64 i = 0; i < param_count; i += 1, param += 1){
+          SYMS_String8 param_str = syms_string_from_type(arena, *param);
+          SYMS_String8 param_str_trimmed = syms_str8_skip_chop_whitespace(param_str);
+          syms_string_list_push(arena, out, param_str_trimmed);
+          if (i + 1 < param_count){
+            syms_string_list_push(arena, out, syms_str8_lit(", "));
+          }
+        }
+        syms_string_list_push(arena, out, syms_str8_lit(")"));
+      }
+      
+      syms_rhs_string_from_type__internal(arena, type->direct_type, out, 2);
+    }break;
+    
+    case SYMS_TypeKind_Label:
+    {
+      // TODO(allen): ???
+    }break;
+  }
+}
+
+#endif //SYMS_TYPE_GRAPH_C
