- Updated Sysmlib that can handle > 4gb pdb files

diff --git a/Engine/Source/Runtime/SymsLib/syms/core/group/syms_type_graph.h b/Engine/Source/Runtime/SymsLib/syms/core/group/syms_type_graph.h
new file mode 100644
--- /dev/null
+++ b/Engine/Source/Runtime/SymsLib/syms/core/group/syms_type_graph.h
@@ -0,0 +1,241 @@
+// Copyright Epic Games, Inc. All Rights Reserved.
+/* date = February 16th 2022 1:40 pm */
+
+#ifndef SYMS_TYPE_GRAPH_H
+#define SYMS_TYPE_GRAPH_H
+
+////////////////////////////////
+// NOTE(allen): Syms Type Graph
+
+typedef struct SYMS_TypeMember{
+  SYMS_MemKind kind;
+  SYMS_MemVisibility visibility;
+  SYMS_MemFlags flags;
+  SYMS_String8 name;
+  SYMS_U32 off;
+  SYMS_U32 virtual_off;
+  struct SYMS_TypeNode *type;
+} SYMS_TypeMember;
+
+typedef struct SYMS_TypeSrcCoord{
+  SYMS_USID usid;
+  SYMS_FileID file_id;
+  SYMS_U32 line;
+  SYMS_U32 col;
+} SYMS_TypeSrcCoord;
+
+typedef struct SYMS_TypeMemberArray{
+  SYMS_TypeMember *mems;
+  SYMS_U64 count;
+} SYMS_TypeMemberArray;
+
+typedef struct SYMS_TypeNode{
+  // SYMS_TypeNode extends and completes the information from SYMS_TypeInfo.
+  // See SYMS_TypeInfo for more interpretation info.
+  
+  SYMS_TypeKind kind;
+  SYMS_String8 name;
+  SYMS_U64 byte_size;
+  
+  // when non-null contains the source location of the type's definition
+  SYMS_TypeSrcCoord *src_coord;
+  
+  // (in addition to interpretations of SYMS_TypeInfo 'direct_type')
+  //  SYMS_TypeKind_Forward*   -> the concrete type referenced by the forward reference
+  struct SYMS_TypeNode *direct_type;
+  
+  // 'this_type' meaning depends on kind:
+  //  SYMS_TypeKind_MemberPtr  -> the container type of the member pointer
+  //  SYMS_TypeKind_Proc       -> if non-nil this is the type of an implicit 'this' in a C++ method
+  struct SYMS_TypeNode *this_type;
+  
+  union{
+    // kind: SYMS_TypeKind_Modifier
+    SYMS_TypeModifiers mods;
+    
+    // kind: SYMS_TypeKind_Array
+    SYMS_U64 array_count;
+    
+    // kind: SYMS_TypeKind_Proc
+    struct{
+      struct SYMS_TypeNode **params;
+      SYMS_U64 param_count;
+    } proc;
+    
+    // opaque pointer for lazy eval attachments to the type node.
+    void *lazy_ptr;
+  };
+} SYMS_TypeNode;
+
+typedef struct SYMS_TypeContentNode{
+  struct SYMS_TypeContentNode *next;
+  SYMS_String8 key;
+  SYMS_U64 hash;
+  SYMS_TypeNode *type;
+} SYMS_TypeContentNode;
+
+typedef struct SYMS_TypeContentBuckets{
+  SYMS_TypeContentNode **buckets;
+  SYMS_U64 bucket_count;
+} SYMS_TypeContentBuckets;
+
+typedef struct SYMS_TypeUSIDNode{
+  struct SYMS_TypeUSIDNode *next;
+  SYMS_USID key;
+  SYMS_TypeNode *type;
+} SYMS_TypeUSIDNode;
+
+typedef struct SYMS_TypeUSIDBuckets{
+  SYMS_TypeUSIDNode **buckets;
+  SYMS_U64 bucket_count;
+} SYMS_TypeUSIDBuckets;
+
+typedef struct SYMS_TypeChainNode{
+  struct SYMS_TypeChainNode *next;
+  SYMS_TypeNode *type;
+} SYMS_TypeChainNode;
+
+typedef struct SYMS_TypeChain{
+  SYMS_TypeChainNode *first;
+  SYMS_TypeChainNode *last;
+  SYMS_U64 count;
+} SYMS_TypeChain;
+
+typedef struct SYMS_TypeNameNode{
+  struct SYMS_TypeNameNode *next;
+  SYMS_U8 *name_ptr;
+  SYMS_TypeChain chain;
+} SYMS_TypeNameNode;
+
+typedef struct SYMS_TypeNameBuckets{
+  SYMS_TypeNameNode **buckets;
+  SYMS_U64 bucket_count;
+} SYMS_TypeNameBuckets;
+
+typedef struct SYMS_TypeGraph{
+  SYMS_Arena *arena;
+  SYMS_StringCons *string_cons;
+  SYMS_U64 address_size;
+  SYMS_TypeUSIDBuckets type_usid_buckets;
+  SYMS_TypeContentBuckets type_content_buckets;
+  SYMS_TypeNameBuckets type_name_buckets;
+} SYMS_TypeGraph;
+
+////////////////////////////////
+// NOTE(allen): Nil Type Node
+
+SYMS_READ_ONLY SYMS_GLOBAL SYMS_TypeNode syms_type_graph_nil = {
+  SYMS_TypeKind_Null,     // kind
+  {(SYMS_U8*)"(nil)", 5}, // name
+  0,                      // byte_size
+  0,                      // src_coord
+  &syms_type_graph_nil,   // direct_type
+  &syms_type_graph_nil,   // this_type
+};
+
+////////////////////////////////
+// NOTE(allen): Type Graph Setup Functions
+
+// The graph_arena is a "permanent arena" that is, it is used
+// for all memory allocations on the graph. It should not be
+// popped or released until the graph is no longer in use. It
+// should not be a scratch arena.
+// Similarly with the string cons structure.
+
+SYMS_API void            syms_type_graph_init(SYMS_TypeGraph *graph,
+                                              SYMS_Arena *graph_arena,
+                                              SYMS_StringCons *graph_string_cons,
+                                              SYMS_U64 address_size);
+
+////////////////////////////////
+// NOTE(allen): Type Content Hash Functions
+
+SYMS_API SYMS_U64       syms_type_content_hash(SYMS_String8 data);
+SYMS_API SYMS_TypeNode* syms_type_from_content(SYMS_TypeContentBuckets *buckets, SYMS_String8 data);
+SYMS_API SYMS_String8   syms_type_content_buckets_insert(SYMS_Arena *arena, SYMS_TypeContentBuckets *buckets,
+                                                         SYMS_String8 key, SYMS_TypeNode *type);
+
+////////////////////////////////
+// NOTE(allen): Type USID Hash Functions
+
+SYMS_API SYMS_U64       syms_type_usid_hash(SYMS_USID usid);
+SYMS_API SYMS_TypeNode* syms_type_from_usid(SYMS_TypeUSIDBuckets *buckets,
+                                            SYMS_USID usid);
+SYMS_API void           syms_type_usid_buckets_insert(SYMS_Arena *arena,
+                                                      SYMS_TypeUSIDBuckets *buckets,
+                                                      SYMS_USID key,
+                                                      SYMS_TypeNode *type);
+
+////////////////////////////////
+// NOTE(allen): Type Name Hash Functions
+
+SYMS_API SYMS_U64        syms_type_name_hash(SYMS_U8 *ptr);
+SYMS_API SYMS_TypeChain* syms_type_chain_from_name(SYMS_TypeNameBuckets *buckets,
+                                                   SYMS_U8 *name_ptr);
+SYMS_API void            syms_type_name_buckets_insert(SYMS_Arena *arena,
+                                                       SYMS_TypeNameBuckets *buckets,
+                                                       SYMS_U8 *name_ptr, SYMS_TypeNode *type);
+
+////////////////////////////////
+// NOTE(allen): Type Name Lookup Functions
+
+SYMS_API SYMS_TypeChain syms_type_from_name(SYMS_TypeGraph *graph, SYMS_String8 name);
+
+////////////////////////////////
+// NOTE(allen): Type Node Constructors
+
+SYMS_API SYMS_TypeNode* syms_type_basic(SYMS_TypeGraph *graph,
+                                        SYMS_TypeKind basic_kind, SYMS_U64 size,
+                                        SYMS_String8 name);
+SYMS_API SYMS_TypeNode* syms_type_mod_from_type(SYMS_TypeGraph *graph,
+                                                SYMS_TypeNode *type, SYMS_TypeModifiers mods);
+SYMS_API SYMS_TypeNode* syms_type_ptr_from_type(SYMS_TypeGraph *graph,
+                                                SYMS_TypeKind ptr_kind, SYMS_TypeNode *type);
+SYMS_API SYMS_TypeNode* syms_type_array_from_type(SYMS_TypeGraph *graph,
+                                                  SYMS_TypeNode *type, SYMS_U64 count);
+SYMS_API SYMS_TypeNode* syms_type_proc_from_type(SYMS_TypeGraph *graph,
+                                                 SYMS_TypeNode *ret_type,
+                                                 SYMS_TypeNode *this_type,
+                                                 SYMS_TypeNode **params, SYMS_U64 count);
+SYMS_API SYMS_TypeNode* syms_type_member_ptr_from_type(SYMS_TypeGraph *graph,
+                                                       SYMS_TypeNode *container,
+                                                       SYMS_TypeNode *type);
+
+////////////////////////////////
+// NOTE(allen): Type Info Operators
+
+SYMS_API SYMS_TypeNode* syms_type_resolved(SYMS_TypeNode *type);
+SYMS_API SYMS_B32       syms_type_node_match(SYMS_TypeNode *l, SYMS_TypeNode *r);
+
+SYMS_API SYMS_TypeNode* syms_type_resolve_enum_to_basic(SYMS_TypeGraph *graph, SYMS_TypeNode *t);
+
+SYMS_API SYMS_TypeNode* syms_type_promoted_from_type_node(SYMS_TypeGraph *graph,
+                                                          SYMS_TypeNode *c);
+SYMS_API SYMS_TypeNode* syms_type_auto_casted_from_type_nodes(SYMS_TypeGraph *graph,
+                                                              SYMS_TypeNode *l,
+                                                              SYMS_TypeNode *r);
+
+////////////////////////////////
+// NOTE(allen): Type Stringizing
+
+SYMS_API SYMS_String8   syms_string_from_type(SYMS_Arena *arena, SYMS_TypeNode *type);
+
+SYMS_API void           syms_lhs_string_from_type(SYMS_Arena *arena, SYMS_TypeNode *type,
+                                                  SYMS_String8List *out);
+SYMS_API void           syms_rhs_string_from_type(SYMS_Arena *arena, SYMS_TypeNode *type,
+                                                  SYMS_String8List *out);
+
+SYMS_API void           syms_lhs_string_from_type_skip_return(SYMS_Arena *arena,
+                                                              SYMS_TypeNode *type,
+                                                              SYMS_String8List *out);
+
+SYMS_API void           syms_lhs_string_from_type__internal(SYMS_Arena *arena,
+                                                            SYMS_TypeNode *type,
+                                                            SYMS_String8List *out,
+                                                            SYMS_U32 prec, SYMS_B32 skip);
+SYMS_API void           syms_rhs_string_from_type__internal(SYMS_Arena *arena,
+                                                            SYMS_TypeNode *type,
+                                                            SYMS_String8List *out,
+                                                            SYMS_U32 prec);
+
+#endif //SYMS_TYPE_GRAPH_H
diff --git a/Engine/Source/Runtime/SymsLib/syms/core/mach/syms_mach.h b/Engine/Source/Runtime/SymsLib/syms/core/mach/syms_mach.h
--- a/Engine/Source/Runtime/SymsLib/syms/core/mach/syms_mach.h
+++ b/Engine/Source/Runtime/SymsLib/syms/core/mach/syms_mach.h
@@ -1,4 +1,5 @@
 // Copyright Epic Games, Inc. All Rights Reserved.
+
 #ifndef SYMS_MACH_H
 #define SYMS_MACH_H
 
@@ -51,23 +52,6 @@ typedef SYMS_S32 SYMS_MachVMProt;
 
 #include "syms/core/generated/syms_meta_mach.h"
 
-////////////////////////////////
-//~ NOTE(rjf): More Types & Constants from loader.h
-
-typedef struct SYMS_MachBuildVersionCommand{
-  SYMS_U32 cmd;
-  SYMS_U32 cmdsize;
-  SYMS_MachPlatformType platform;
-  SYMS_U32 minos;
-  SYMS_U32 sdk;
-  SYMS_U32 ntools;
-} SYMS_MachBuildVersionCommand;
-
-typedef struct SYMS_MachBuildToolVersion{
-  SYMS_MachToolType tool;
-  SYMS_U32 version;
-} SYMS_MachBuildToolVersion;
-
 ////////////////////////////////
 //~ NOTE(allen): Types & Constants from nlist.h
 
diff --git a/Engine/Source/Runtime/SymsLib/syms/core/mach/syms_mach_parser.c b/Engine/Source/Runtime/SymsLib/syms/core/mach/syms_mach_parser.c
--- a/Engine/Source/Runtime/SymsLib/syms/core/mach/syms_mach_parser.c
+++ b/Engine/Source/Runtime/SymsLib/syms/core/mach/syms_mach_parser.c
@@ -1,4 +1,5 @@
 // Copyright Epic Games, Inc. All Rights Reserved.
+
 #ifndef SYMS_MACH_PARSER_C
 #define SYMS_MACH_PARSER_C
 
@@ -72,6 +73,8 @@ syms_mach_bin_from_base_range(SYMS_Arena *arena, void *base, SYMS_U64Range range
     
     SYMS_U64Range bind_ranges[SYMS_MachBindTable_COUNT]; syms_memzero_struct(&bind_ranges[0]);
     SYMS_U64Range export_range = syms_make_u64_range(0,0);
+
+    SYMS_MachSymtabCommand symtab; syms_memzero_struct(&symtab);
     
     SYMS_MachDylibList dylib_list; syms_memzero_struct(&dylib_list);
     
@@ -95,8 +98,47 @@ syms_mach_bin_from_base_range(SYMS_Arena *arena, void *base, SYMS_U64Range range
       SYMS_MachLoadCommand lc = {0};
       syms_based_range_read_struct(base, range, cmd_off, &lc);
       switch (lc.type){
-        // TODO(allen): SYMS_MachSegmentCommand32?
-        
+        case SYMS_MachLoadCommandType_SYMTAB: 
+        {
+          syms_based_range_read_struct(base, range, cmd_off, &symtab);
+        }break;
+
+        case SYMS_MachLoadCommandType_SEGMENT: 
+        {
+          SYMS_MachSegmentCommand32 segment_command32 = {0};
+          syms_based_range_read_struct(base, range, cmd_off, &segment_command32);
+          if (is_swapped){
+            syms_bswap_in_place(SYMS_MachSegmentCommand32, &segment_command32);
+          }
+          SYMS_U64 after_seg_off = cmd_off + sizeof(SYMS_MachSegmentCommand32);
+
+          // push segment node
+          SYMS_MachSegmentNode *segment_node = syms_push_array_zero(scratch.arena, SYMS_MachSegmentNode, 1);
+          syms_mach_segment64_from_segment32(&segment_node->data, &segment_command32);
+          SYMS_QueuePush(segment_first, segment_last, segment_node);
+          segment_count += 1;
+
+          // loop over sections
+          SYMS_U64 next_sec_off = after_seg_off;
+          SYMS_U64 sec_count = segment_command32.nsects;
+          for (SYMS_U32 k = 0; k < sec_count; k += 1){
+            // read section 64
+            SYMS_U64 sec_off = next_sec_off;
+            SYMS_MachSection32 section32 = {0};
+            syms_based_range_read_struct(base, range, sec_off, &section32);
+            if (is_swapped){
+              syms_bswap_in_place(SYMS_MachSection32, &section32);
+            }
+            next_sec_off = sec_off + sizeof(SYMS_MachSection32);
+
+            // push section node
+            SYMS_MachSectionNode *section_node = syms_push_array_zero(scratch.arena, SYMS_MachSectionNode, 1);
+            syms_mach_section64_from_section32(&section_node->data, &section32);
+            SYMS_QueuePush(section_first, section_last, section_node);
+            section_count += 1;
+          }
+        }break;
+
         case SYMS_MachLoadCommandType_SEGMENT_64:
         {
           // read segment 64
@@ -191,13 +233,17 @@ syms_mach_bin_from_base_range(SYMS_Arena *arena, void *base, SYMS_U64Range range
     }
     
     //- fill result
-    result                = syms_push_array(arena, SYMS_MachBinAccel, 1);
-    result->format        = SYMS_FileFormat_MACH;
-    result->arch          = syms_mach_arch_from_cputype(header.cputype);
-    result->segment_count = segment_count;
-    result->segments      = segments;
-    result->section_count = section_count;
-    result->sections      = sections;
+    result                     = syms_push_array(arena, SYMS_MachBinAccel, 1);
+    result->load_command_count = header.ncmds;
+    result->load_commands      = syms_make_u64_inrange(range, after_header_off, header.sizeofcmds);
+    result->format             = SYMS_FileFormat_MACH;
+    result->arch               = syms_mach_arch_from_cputype(header.cputype);
+    result->is_swapped         = is_swapped;
+    result->symtab             = symtab;
+    result->segment_count      = segment_count;
+    result->segments           = segments;
+    result->section_count      = section_count;
+    result->sections           = sections;
     syms_memmove(&result->bind_ranges[0], &bind_ranges[0], sizeof(bind_ranges));
     result->export_range  = export_range;
     { // convert dylib list to array
@@ -383,14 +429,17 @@ syms_mach_sec_info_array_from_bin(SYMS_Arena *arena, SYMS_String8 data, SYMS_Mac
   SYMS_MachSection64 *mach_sec_opl = bin->sections + array.count;
   SYMS_SecInfo *sec_info = array.sec_info;
   for (; mach_sec < mach_sec_opl; sec_info += 1, mach_sec += 1){
-    SYMS_U8 *ptr = mach_sec->sectname;
-    SYMS_U8 *opl = mach_sec->sectname + SYMS_ARRAY_SIZE(mach_sec->sectname);
+    SYMS_U8 *name_ptr = mach_sec->sectname;
+    SYMS_U8 *ptr = name_ptr;
+    SYMS_U8 *opl = name_ptr + SYMS_ARRAY_SIZE(mach_sec->sectname);
     for (;ptr < opl && *ptr != 0; ptr += 1);
     
-    sec_info->name = syms_str8_range(mach_sec->sectname, ptr);
+    SYMS_String8 name = syms_str8_range(name_ptr, ptr);
+    
+    sec_info->name = syms_push_string_copy(arena, name);
     // TODO(allen): figure out when these ranges actually apply and when the section isn't
     // there one side or the other.
-    sec_info->vrange = syms_make_u64_range(mach_sec->addr,   mach_sec->addr +   mach_sec->size);
+    sec_info->vrange = syms_make_u64_range(mach_sec->addr,   mach_sec->addr   + mach_sec->size);
     sec_info->frange = syms_make_u64_range(mach_sec->offset, mach_sec->offset + mach_sec->size);
   }
   return array;
diff --git a/Engine/Source/Runtime/SymsLib/syms/core/mach/syms_mach_parser.h b/Engine/Source/Runtime/SymsLib/syms/core/mach/syms_mach_parser.h
--- a/Engine/Source/Runtime/SymsLib/syms/core/mach/syms_mach_parser.h
+++ b/Engine/Source/Runtime/SymsLib/syms/core/mach/syms_mach_parser.h
@@ -1,4 +1,5 @@
 // Copyright Epic Games, Inc. All Rights Reserved.
+
 #ifndef SYMS_MACH_PARSER_H
 #define SYMS_MACH_PARSER_H
 
@@ -40,16 +41,23 @@ typedef struct SYMS_MachBinAccel{
   SYMS_FileFormat format;
   
   SYMS_Arch arch;
+
+  SYMS_B32 is_swapped;
+
+  SYMS_U32 load_command_count;
+  SYMS_U64Range load_commands;
+
+  SYMS_MachSymtabCommand symtab;
   
   SYMS_MachSegmentCommand64 *segments;
   SYMS_U32 segment_count;
   
   SYMS_MachSection64 *sections;
   SYMS_U32 section_count;
-
+  
   SYMS_U64Range bind_ranges[SYMS_MachBindTable_COUNT];
   SYMS_U64Range export_range;
-
+  
   SYMS_U32 dylib_count;
   struct SYMS_MachParsedDylib *dylibs;
 } SYMS_MachBinAccel;
@@ -113,7 +121,7 @@ typedef struct SYMS_MachExport
   SYMS_U8 child_count;
   
   SYMS_B32 is_export_info;
-
+  
   SYMS_String8 name;
   SYMS_U64 flags;
   SYMS_U64 address;
diff --git a/Engine/Source/Runtime/SymsLib/syms/core/pdb/syms_msf.h b/Engine/Source/Runtime/SymsLib/syms/core/pdb/syms_msf.h
--- a/Engine/Source/Runtime/SymsLib/syms/core/pdb/syms_msf.h
+++ b/Engine/Source/Runtime/SymsLib/syms/core/pdb/syms_msf.h
@@ -10,8 +10,8 @@
 #define SYMS_MSF_INVALID_STREAM_NUMBER SYMS_U16_MAX
 typedef SYMS_U16 SYMS_MsfStreamNumber; 
 
-SYMS_GLOBAL char syms_msf20_magic[] = "Microsoft C/C++ program database 2.00\r\n\x1aJG\0\0";
-SYMS_GLOBAL char syms_msf70_magic[] = "Microsoft C/C++ MSF 7.00\r\n\032DS\0\0";
+SYMS_READ_ONLY SYMS_GLOBAL char syms_msf20_magic[] = "Microsoft C/C++ program database 2.00\r\n\x1aJG\0\0";
+SYMS_READ_ONLY SYMS_GLOBAL char syms_msf70_magic[] = "Microsoft C/C++ MSF 7.00\r\n\032DS\0\0";
 
 #define SYMS_MSF20_MAGIC_SIZE 44
 #define SYMS_MSF70_MAGIC_SIZE 32
diff --git a/Engine/Source/Runtime/SymsLib/syms/core/pdb/syms_msf_parser.c b/Engine/Source/Runtime/SymsLib/syms/core/pdb/syms_msf_parser.c
--- a/Engine/Source/Runtime/SymsLib/syms/core/pdb/syms_msf_parser.c
+++ b/Engine/Source/Runtime/SymsLib/syms/core/pdb/syms_msf_parser.c
@@ -1,4 +1,5 @@
 // Copyright Epic Games, Inc. All Rights Reserved.
+
 #ifndef SYMS_MSF_PARSER_C
 #define SYMS_MSF_PARSER_C
 
@@ -90,7 +91,7 @@ syms_msf_accel_from_data(SYMS_Arena *arena, SYMS_String8 data){
       // setup important sizes and counts
       SYMS_U32 size_of_block_index = header.index_size;
       
-      SYMS_U32 file_size = data.size;
+      SYMS_U64 file_size = data.size;
       SYMS_U32 block_size = SYMS_ClampTop(header.block_size, file_size);
       
       SYMS_U32 block_count_in_directory = SYMS_CeilIntegerDiv(directory_size, block_size);
@@ -111,7 +112,7 @@ syms_msf_accel_from_data(SYMS_Arena *arena, SYMS_String8 data){
         directory_map_block_skip_size = SYMS_MSF20_MAGIC_SIZE + SYMS_MEMBER_OFFSET(SYMS_MsfHeader20, directory_map);
       }
       else{
-        SYMS_U32 super_map_off = SYMS_MSF70_MAGIC_SIZE + SYMS_MEMBER_OFFSET(SYMS_MsfHeader70, directory_super_map);
+        SYMS_U64 super_map_off = SYMS_MSF70_MAGIC_SIZE + SYMS_MEMBER_OFFSET(SYMS_MsfHeader70, directory_super_map);
         directory_super_map = (SYMS_U32*)(data.str + super_map_off);
       }
       
@@ -134,7 +135,7 @@ syms_msf_accel_from_data(SYMS_Arena *arena, SYMS_String8 data){
           goto read_directory_done;
         }
         
-        SYMS_U32 directory_map_block_off = directory_map_block_index*block_size;
+        SYMS_U64 directory_map_block_off = (SYMS_U64)(directory_map_block_index)*block_size;
         SYMS_U8 *directory_map_block_base = data.str + directory_map_block_off;
         
         // clamp index count by end of directory
@@ -159,7 +160,7 @@ syms_msf_accel_from_data(SYMS_Arena *arena, SYMS_String8 data){
             goto read_directory_done;
           }
           
-          SYMS_U32 directory_block_off = directory_block_index*block_size;
+          SYMS_U64 directory_block_off = (SYMS_U64)(directory_block_index)*block_size;
           SYMS_U8 *directory_block_base = data.str + directory_block_off;
           
           // clamp copy size by end of directory
@@ -293,14 +294,13 @@ syms_msf_read(SYMS_String8 data, SYMS_MsfAccel *msf,
   SYMS_MsfStreamInfo stream_info = syms_msf_stream_info_from_sn(msf, sn);
   if (stream_info.sn != 0 && off + size <= stream_info.size && size > 0){
     // copy block-by-block
-    SYMS_U32 file_size = data.size;
+    SYMS_U64 file_size = data.size;
     SYMS_U32 block_size = msf->header.block_size;
     SYMS_U32 size_of_block_index = msf->header.index_size;
-    SYMS_U32 block_count_in_whole_file = SYMS_ClampTop(msf->header.block_count, file_size/block_size);
+    SYMS_U32 block_count_in_whole_file_max = SYMS_CeilIntegerDiv(file_size, block_size);
+    SYMS_U32 block_count_in_whole_file = SYMS_ClampTop(msf->header.block_count, block_count_in_whole_file_max);
     SYMS_U32 block_count_in_stream = SYMS_CeilIntegerDiv(stream_info.size, block_size);
     
-    SYMS_U8 *file_data = data.str;
-    
     SYMS_U32 completed_amount = 0;
     for (;;){
       // offset-in-stream -> part-index
@@ -318,14 +318,14 @@ syms_msf_read(SYMS_String8 data, SYMS_MsfAccel *msf,
       }
       
       //  block-index -> range-in-file
-      SYMS_U32 src_block_base = src_block_index*block_size;
+      SYMS_U64 src_block_base = (SYMS_U64)(src_block_index)*block_size;
       SYMS_U32 relative_off = remaining_off%block_size;
       SYMS_U32 relative_opl = SYMS_ClampTop(relative_off + size - completed_amount, block_size);
       SYMS_U32 contiguous_size = relative_opl - relative_off;
       SYMS_ASSERT(src_block_base + relative_opl <= file_size);
       
       // copy & advance
-      syms_memmove((SYMS_U8*)out + completed_amount, file_data + src_block_base + relative_off, contiguous_size);
+      syms_memmove((SYMS_U8*)out + completed_amount, data.str + src_block_base + relative_off, contiguous_size);
       completed_amount += contiguous_size;
       if (completed_amount >= size){
         result = syms_true;
@@ -413,14 +413,13 @@ syms_msf_read_zstring_in_range(SYMS_Arena *arena, SYMS_String8 data, SYMS_MsfAcc
   SYMS_MsfStreamInfo stream_info = syms_msf_stream_info_from_sn(msf, range.sn);
   if (stream_info.sn != 0 && off < stream_info.size && off < max_off){
     // scan block-by-block
-    SYMS_U32 file_size = data.size;
+    SYMS_U64 file_size = data.size;
     SYMS_U32 block_size = msf->header.block_size;
     SYMS_U32 size_of_block_index = msf->header.index_size;
-    SYMS_U32 block_count_in_whole_file = SYMS_ClampTop(msf->header.block_count, file_size/block_size);
+    SYMS_U32 block_count_in_whole_file_max = SYMS_CeilIntegerDiv(file_size, block_size);
+    SYMS_U32 block_count_in_whole_file = SYMS_ClampTop(msf->header.block_count, block_count_in_whole_file_max);
     SYMS_U32 block_count_in_stream = SYMS_CeilIntegerDiv(stream_info.size, block_size);
     
-    SYMS_U8 *file_data = data.str;
-    
     SYMS_U32 max_scan_amount = max_off - off;
     SYMS_U32 scanned_amount = 0;
     for (;;){
@@ -446,14 +445,14 @@ syms_msf_read_zstring_in_range(SYMS_Arena *arena, SYMS_String8 data, SYMS_MsfAcc
       }
       
       // block-index -> range-in-file
-      SYMS_U32 src_block_base = src_block_index*block_size;
+      SYMS_U64 src_block_base = (SYMS_U64)(src_block_index)*block_size;
       SYMS_U32 relative_off = remaining_off%block_size;
       SYMS_U32 remaining_max_size = max_off - remaining_off;
       SYMS_U32 relative_opl = SYMS_ClampTop(relative_off + remaining_max_size, block_size);
       SYMS_ASSERT(src_block_base + block_size <= file_size);
       
       // scan
-      SYMS_U8 *start = file_data + src_block_base + relative_off;
+      SYMS_U8 *start = data.str + src_block_base + relative_off;
       SYMS_U8 *opl = start + relative_opl - relative_off;
       SYMS_U8 *ptr = start;
       for (;ptr < opl && *ptr != 0; ptr += 1);
diff --git a/Engine/Source/Runtime/SymsLib/syms/core/pdb/syms_pdb_parser.c b/Engine/Source/Runtime/SymsLib/syms/core/pdb/syms_pdb_parser.c
--- a/Engine/Source/Runtime/SymsLib/syms/core/pdb/syms_pdb_parser.c
+++ b/Engine/Source/Runtime/SymsLib/syms/core/pdb/syms_pdb_parser.c
@@ -1,4 +1,5 @@
 // Copyright Epic Games, Inc. All Rights Reserved.
+
 #ifndef SYMS_PDB_PARSER_C
 #define SYMS_PDB_PARSER_C
 
@@ -268,7 +269,7 @@ syms_pdb_types_from_name(SYMS_Arena *arena, SYMS_String8 data, SYMS_PdbDbgAccel
 //~ NOTE(allen): PDB GSI
 
 SYMS_API SYMS_PdbGsiAccel
-syms_pdb_gsi_accel_from_range(SYMS_Arena *arena, SYMS_String8 data, SYMS_MsfAccel *accel, SYMS_MsfRange range){
+syms_pdb_gsi_accel_from_range(SYMS_Arena *arena, SYMS_String8 data, SYMS_MsfAccel *msf, SYMS_MsfRange range){
   SYMS_ProfBegin("syms_pdb_parse_info");
   
   // get scratch
@@ -276,7 +277,7 @@ syms_pdb_gsi_accel_from_range(SYMS_Arena *arena, SYMS_String8 data, SYMS_MsfAcce
   // check gsi info
   SYMS_B32 has_gsi = syms_false;
   SYMS_PdbGsiHeader header = {0};
-  if (syms_msf_read_struct_in_range(data, accel, range, 0, &header)){
+  if (syms_msf_read_struct_in_range(data, msf, range, 0, &header)){
     if (header.sig == SYMS_PdbGsiSignature_Basic &&
         header.ver == SYMS_PdbGsiVersion_V70 &&
         header.num_buckets != 0){
@@ -306,12 +307,12 @@ syms_pdb_gsi_accel_from_range(SYMS_Arena *arena, SYMS_String8 data, SYMS_MsfAcce
     
     // read bitmask
     SYMS_MsfRange bitmask_range = syms_msf_sub_range(range, bitmask_off, bitmask_size);
-    SYMS_String8 bitmask = syms_msf_read_whole_range(scratch.arena, data, accel, bitmask_range);
+    SYMS_String8 bitmask = syms_msf_read_whole_range(scratch.arena, data, msf, bitmask_range);
     SYMS_U32 bitmask_count_clamped = bitmask.size/4;
     
     // read packed offsets
     SYMS_MsfRange offsets_range = syms_msf_sub_range(range, offsets_off, range.size - offsets_off);
-    SYMS_String8 packed_offsets = syms_msf_read_whole_range(scratch.arena, data, accel, offsets_range);
+    SYMS_String8 packed_offsets = syms_msf_read_whole_range(scratch.arena, data, msf, offsets_range);
     SYMS_U32 packed_offsets_count = packed_offsets.size/4;
     
     // unpack offsets
@@ -354,7 +355,7 @@ syms_pdb_gsi_accel_from_range(SYMS_Arena *arena, SYMS_String8 data, SYMS_MsfAcce
     
     // read hash records
     SYMS_MsfRange hash_record_range = syms_msf_sub_range(range, hash_record_array_off, header.hr_len);
-    SYMS_String8 hash_records = syms_msf_read_whole_range(scratch.arena, data, accel, hash_record_range);
+    SYMS_String8 hash_records = syms_msf_read_whole_range(scratch.arena, data, msf, hash_record_range);
     SYMS_U32 entry_count = hash_records.size/sizeof(SYMS_PdbGsiHashRecord);
     
     // build bucket chains
@@ -1695,6 +1696,11 @@ syms_pdb_sym_accel_from_range(SYMS_Arena *arena, SYMS_String8 data, SYMS_MsfAcce
   SYMS_PdbStubRef *var_last = 0;
   SYMS_U64 var_count = 0;
   
+  // thread vars
+  SYMS_PdbStubRef *tls_var_first = 0;
+  SYMS_PdbStubRef *tls_var_last = 0;
+  SYMS_U64 tls_var_count = 0;
+  
   // thunk
   SYMS_PdbStubRef *thunk_first = 0;
   SYMS_PdbStubRef *thunk_last = 0;
@@ -1706,7 +1712,7 @@ syms_pdb_sym_accel_from_range(SYMS_Arena *arena, SYMS_String8 data, SYMS_MsfAcce
   SYMS_U64 pub_count = 0;
   
   // parse loop
-  SYMS_PdbStub *defrange_collector_stub = 0;
+  SYMS_PdbStub *defrange_collector_stub = &root;
   SYMS_PdbStubRef *stack = 0;
   SYMS_PdbStubRef *stack_free = 0;
   SYMS_U32 cursor = 0;
@@ -1721,8 +1727,7 @@ syms_pdb_sym_accel_from_range(SYMS_Arena *arena, SYMS_String8 data, SYMS_MsfAcce
     
     // init stub
     SYMS_U32 symbol_off = element.range.off - 4;
-    SYMS_PdbStub *stub = syms_push_array(arena, SYMS_PdbStub, 1);
-    syms_memzero_struct(stub);
+    SYMS_PdbStub *stub = syms_push_array_zero(arena, SYMS_PdbStub, 1);
     SYMS_QueuePush_N(first, last, stub, bucket_next);
     all_count += 1;
     stub->off = symbol_off;
@@ -1755,7 +1760,6 @@ syms_pdb_sym_accel_from_range(SYMS_Arena *arena, SYMS_String8 data, SYMS_MsfAcce
         var_count += 1;
       }break;
       
-#if 0
       case SYMS_CvSymKind_LTHREAD32:
       case SYMS_CvSymKind_GTHREAD32:
       {
@@ -1767,10 +1771,9 @@ syms_pdb_sym_accel_from_range(SYMS_Arena *arena, SYMS_String8 data, SYMS_MsfAcce
         // push onto var list
         SYMS_PdbStubRef *ref = syms_pdb_alloc_ref(scratch.arena, &stack_free);
         ref->stub = stub;
-        SYMS_QueuePush(var_first, var_last, ref);
-        var_count += 1;
+        SYMS_QueuePush(tls_var_first, tls_var_last, ref);
+        tls_var_count += 1;
       }break;
-#endif
       
       // general local variable
       case SYMS_CvSymKind_LOCAL:
@@ -1784,6 +1787,17 @@ syms_pdb_sym_accel_from_range(SYMS_Arena *arena, SYMS_String8 data, SYMS_MsfAcce
         defrange_collector_stub = stub;
       }break;
       
+      case SYMS_CvSymKind_FILESTATIC:
+      {
+        SYMS_U32 name_off = sizeof(SYMS_CvFileStatic);
+        SYMS_String8 name = syms_msf_read_zstring_in_range(arena, data, msf, element.range, name_off);
+        
+        stub->name = name;
+        
+        preserve_defrange_collection = syms_true;
+        defrange_collector_stub = stub;
+      }break;
+      
       case SYMS_CvSymKind_DEFRANGE_2005:
       case SYMS_CvSymKind_DEFRANGE2_2005:
       case SYMS_CvSymKind_DEFRANGE:
@@ -1832,6 +1846,9 @@ syms_pdb_sym_accel_from_range(SYMS_Arena *arena, SYMS_String8 data, SYMS_MsfAcce
       case SYMS_CvSymKind_LPROC32:
       case SYMS_CvSymKind_GPROC32:
       {
+        // TODO(allen): (element.range.off + element.range.size) gives the beginning of
+        // the FRAMEPROC - use this to write a fast *PROC32 -> FRAMEPROC helper
+        
         SYMS_U32 name_off = sizeof(SYMS_CvProc32);
         SYMS_String8 name = syms_msf_read_zstring_in_range(arena, data, msf, element.range, name_off);
         
@@ -1869,7 +1886,7 @@ syms_pdb_sym_accel_from_range(SYMS_Arena *arena, SYMS_String8 data, SYMS_MsfAcce
           SYMS_StackPop(stack);
           SYMS_StackPush(stack_free, bucket);
         }
-        parent_for_this_stub = 0;
+        parent_for_this_stub = &root;
       }break;
       
       case SYMS_CvSymKind_PUB32:
@@ -1901,14 +1918,16 @@ syms_pdb_sym_accel_from_range(SYMS_Arena *arena, SYMS_String8 data, SYMS_MsfAcce
     
     // clear defrange collector
     if (!preserve_defrange_collection){
-      defrange_collector_stub = 0;
+      defrange_collector_stub = &root;
     }
     
     // insert into tree
-    if (parent_for_this_stub != 0){
-      SYMS_QueuePush_N(parent_for_this_stub->first, parent_for_this_stub->last, stub, sibling_next);
+    SYMS_ASSERT(parent_for_this_stub != 0);
+    SYMS_QueuePush_N(parent_for_this_stub->first, parent_for_this_stub->last, stub, sibling_next);
+    if (parent_for_this_stub != &root){
+      stub->parent = parent_for_this_stub;
     }
-    if (parent_for_this_stub == &root){
+    else{
       top_count += 1;
     }
     
@@ -1949,6 +1968,17 @@ syms_pdb_sym_accel_from_range(SYMS_Arena *arena, SYMS_String8 data, SYMS_MsfAcce
     }
   }
   
+  // build thread stubs pointer table
+  SYMS_PdbStub **tls_var_stubs = syms_push_array(arena, SYMS_PdbStub*, tls_var_count);
+  {
+    SYMS_PdbStub **ptr = tls_var_stubs;
+    for (SYMS_PdbStubRef *ref = tls_var_first;
+         ref != 0;
+         ref = ref->next, ptr += 1){
+      *ptr = ref->stub;
+    }
+  }
+  
   // build thunk stubs pointer table
   SYMS_PdbStub **thunk_stubs = syms_push_array(arena, SYMS_PdbStub*, thunk_count);
   {
@@ -2003,6 +2033,8 @@ syms_pdb_sym_accel_from_range(SYMS_Arena *arena, SYMS_String8 data, SYMS_MsfAcce
   result->proc_count = proc_count;
   result->var_stubs = var_stubs;
   result->var_count = var_count;
+  result->tls_var_stubs = tls_var_stubs;
+  result->tls_var_count = tls_var_count;
   result->thunk_stubs = thunk_stubs;
   result->thunk_count = thunk_count;
   result->pub_stubs = pub_stubs;
@@ -2042,9 +2074,6 @@ syms_pdb_sym_symbol_kind_from_id(SYMS_String8 data, SYMS_PdbDbgAccel *dbg,
         result = SYMS_SymbolKind_ImageRelativeVariable;
       }break;
       case SYMS_CvSymKind_LOCAL:
-      {
-        result = SYMS_SymbolKind_LocalVariable;
-      }break;
       case SYMS_CvSymKind_REGREL32:
       {
         result = SYMS_SymbolKind_LocalVariable;
@@ -2062,6 +2091,11 @@ syms_pdb_sym_symbol_kind_from_id(SYMS_String8 data, SYMS_PdbDbgAccel *dbg,
       {
         result = SYMS_SymbolKind_Scope;
       }break;
+      case SYMS_CvSymKind_LTHREAD32:
+      case SYMS_CvSymKind_GTHREAD32:
+      {
+        result = SYMS_SymbolKind_TLSVariable;
+      }break;
     }
   }
   
@@ -2203,6 +2237,7 @@ syms_pdb_leaf_accel_from_dbg(SYMS_Arena *arena, SYMS_String8 data, SYMS_PdbDbgAc
           SYMS_PdbStub *fl_stub = syms_push_array(arena, SYMS_PdbStub, 1);
           syms_memzero_struct(fl_stub);
           SYMS_QueuePush_N(parent->first, parent->last, fl_stub, sibling_next);
+          fl_stub->parent = parent;
           fl_stub->off = element.range.off + fl_cursor;
           
           // read internal leaf
@@ -2360,6 +2395,7 @@ syms_pdb_leaf_accel_from_dbg(SYMS_Arena *arena, SYMS_String8 data, SYMS_PdbDbgAc
           SYMS_PdbStub *ml_stub = syms_push_array(arena, SYMS_PdbStub, 1);
           syms_memzero_struct(ml_stub);
           SYMS_QueuePush_N(parent->first, parent->last, ml_stub, sibling_next);
+          ml_stub->parent = parent;
           ml_stub->off = element.range.off + ml_cursor;
           ml_stub->num = virtual_offset;
           
@@ -3212,6 +3248,15 @@ syms_pdb_uid_from_accel(SYMS_PdbUnitAccel *unit){
   return(unit->uid);
 }
 
+SYMS_API SYMS_UnitID
+syms_pdb_tls_var_uid_from_dbg(SYMS_PdbDbgAccel *dbg){
+  // TODO(allen): double check this API design
+  // NOTE(nick): Thread var symbols are stored in global symbol stream
+  // which doesn't abstract well with DWARF model where thread vars
+  // are stored on per compilation units basis.
+  return(SYMS_PdbPseudoUnit_SYM);
+}
+
 SYMS_API SYMS_SymbolIDArray
 syms_pdb_proc_sid_array_from_unit(SYMS_Arena *arena, SYMS_PdbUnitAccel *unit){
   SYMS_SymbolIDArray result = {0};
@@ -3260,6 +3305,31 @@ syms_pdb_var_sid_array_from_unit(SYMS_Arena *arena, SYMS_PdbUnitAccel *unit){
   return(result);
 }
 
+SYMS_API SYMS_SymbolIDArray
+syms_pdb_tls_var_sid_array_from_unit(SYMS_Arena *arena, SYMS_PdbUnitAccel *thread_unit){
+  SYMS_SymbolIDArray result = {0};
+  if (!thread_unit->leaf_set){
+    
+    //- allocate array
+    SYMS_U64 count = thread_unit->tls_var_count;
+    SYMS_SymbolID *ids = syms_push_array(arena, SYMS_SymbolID, count);
+    
+    //- fill array
+    SYMS_SymbolID *id_ptr = ids;
+    SYMS_PdbStub **stub_ptr = thread_unit->tls_var_stubs;
+    SYMS_PdbStub **opl = stub_ptr + count;
+    for (; stub_ptr < opl; id_ptr += 1, stub_ptr += 1){
+      *id_ptr = SYMS_ID_u32_u32(SYMS_PdbSymbolIDKind_Off, (**stub_ptr).off);
+    }
+    
+    //- assemble result
+    result.count = count;
+    result.ids = ids;
+  }
+  
+  return(result);
+}
+
 SYMS_API SYMS_SymbolIDArray
 syms_pdb_type_sid_array_from_unit(SYMS_Arena *arena, SYMS_PdbUnitAccel *unit){
   SYMS_SymbolIDArray result = {0};
@@ -3394,6 +3464,16 @@ syms_pdb_sym_type_from_var_id(SYMS_String8 data, SYMS_PdbDbgAccel *dbg,
         result.sid = SYMS_ID_u32_u32(SYMS_PdbSymbolIDKind_Index, regrel32.itype);
         result.uid = SYMS_PdbPseudoUnit_TPI;
       }break;
+      
+      case SYMS_CvSymKind_LTHREAD32:
+      case SYMS_CvSymKind_GTHREAD32:
+      {
+        SYMS_CvThread32 thread32 = {0};
+        syms_msf_read_struct_in_range(data, msf, element.range, 0, &thread32);
+        
+        result.sid = SYMS_ID_u32_u32(SYMS_PdbSymbolIDKind_Index, thread32.itype);
+        result.uid = SYMS_PdbPseudoUnit_TPI;
+      }break;
     }
   }
   
@@ -3434,6 +3514,162 @@ syms_pdb_sym_voff_from_var_sid(SYMS_String8 data, SYMS_PdbDbgAccel *dbg, SYMS_Pd
   return(result);
 }
 
+SYMS_API SYMS_RegSection
+syms_pdb_reg_section_from_x86_reg(SYMS_CvReg cv_reg){
+  SYMS_RegSection result = {0};
+  if (cv_reg < SYMS_ARRAY_SIZE(syms_reg_slices_X86)){
+    SYMS_RegSlice *slice = syms_reg_slices_X86 + cv_reg;
+    SYMS_RegID reg_id = slice->reg_id;
+    if (reg_id < SYMS_RegX86Code_COUNT){
+      SYMS_RegSection *section = syms_reg_section_X86 + reg_id;
+      result.off = section->off + slice->byte_off;
+      result.size = slice->byte_size;
+    }
+  }
+  return(result);
+}
+
+SYMS_API SYMS_RegSection
+syms_pdb_reg_section_from_x64_reg(SYMS_CvReg cv_reg){
+  SYMS_RegSection result = {0};
+  if (cv_reg < SYMS_ARRAY_SIZE(syms_reg_slices_X64)){
+    SYMS_RegSlice *slice = syms_reg_slices_X64 + cv_reg;
+    SYMS_RegID reg_id = slice->reg_id;
+    if (reg_id < SYMS_RegX64Code_COUNT){
+      SYMS_RegSection *section = syms_reg_section_X64 + reg_id;
+      result.off = section->off + slice->byte_off;
+      result.size = slice->byte_size;
+    }
+  }
+  return(result);
+}
+
+SYMS_API SYMS_RegSection
+syms_pdb_reg_section_from_arch_reg(SYMS_Arch arch, SYMS_CvReg cv_reg){
+  SYMS_RegSection result = {0};
+  switch (arch){
+    case SYMS_Arch_X86: result = syms_pdb_reg_section_from_x86_reg(cv_reg); break;
+    case SYMS_Arch_X64: result = syms_pdb_reg_section_from_x64_reg(cv_reg); break;
+  }
+  return(result);
+}
+
+SYMS_API SYMS_RegSection
+syms_pdb_reg_section_from_framepointer(SYMS_String8 data,  SYMS_PdbDbgAccel *dbg,
+                                       SYMS_MsfRange range, SYMS_PdbStub *framepointer_stub){
+  //- get accelerator
+  SYMS_MsfAccel *msf = dbg->msf;
+  
+  //- get LOCAL
+  SYMS_PdbStub *local_stub = framepointer_stub->parent;
+  SYMS_CvElement local_element = {0};
+  SYMS_B32 got_local = syms_false;
+  if (local_stub != 0){
+    local_element = syms_cv_element(data, msf, range, local_stub->off);
+    if (local_element.kind == SYMS_CvSymKind_LOCAL){
+      got_local = syms_true;
+    }
+  }
+  
+  //- get *PROC32
+  SYMS_CvElement root_element = {0};
+  if (got_local){
+    SYMS_PdbStub *root_stub = local_stub->parent;
+    for (;root_stub->parent != 0;){
+      root_stub = root_stub->parent;
+    }
+    root_element = syms_cv_element(data, msf, range, root_stub->off);
+  }
+  SYMS_B32 got_proc32 = (root_element.kind == SYMS_CvSymKind_LPROC32 ||
+                         root_element.kind == SYMS_CvSymKind_GPROC32);
+  
+  //- get FRAMEPROC
+  SYMS_B32 got_fp_flags = syms_false;
+  SYMS_CvFrameprocFlags fp_flags = 0;
+  if (got_proc32){
+    SYMS_U64 fp_off = root_element.range.off + root_element.range.size;
+    SYMS_CvElement fp_element = syms_cv_element(data, msf, range, fp_off);
+    if (fp_element.kind == SYMS_CvSymKind_FRAMEPROC){
+      SYMS_U64 flags_off = SYMS_MEMBER_OFFSET(SYMS_CvFrameproc, flags);
+      got_fp_flags = syms_msf_read_struct_in_range(data, msf, fp_element.range,
+                                                   flags_off, &fp_flags);
+    }
+  }
+  
+  //- determine the register
+  SYMS_RegSection result = {0};
+  SYMS_Arch arch = syms_pdb_arch_from_dbg(dbg);
+  if (got_fp_flags){
+    // check if parameter
+    SYMS_U64 flags_off = SYMS_MEMBER_OFFSET(SYMS_CvLocal, flags);
+    SYMS_CvLocalFlags local_flags = 0;
+    syms_msf_read_struct_in_range(data, msf, local_element.range, flags_off, &local_flags);
+    SYMS_B32 is_parameter = (local_flags & SYMS_CvLocalFlag_PARAM);
+    
+    // get frame pointer code
+    SYMS_CvEncodedFramePtrReg frame_ptr_reg = 0;
+    if (is_parameter){
+      frame_ptr_reg = SYMS_CvFrameprocFlags_Extract_ParamBasePointer(fp_flags);
+    }
+    else{
+      frame_ptr_reg = SYMS_CvFrameprocFlags_Extract_LocalBasePointer(fp_flags);
+    }
+    
+    // get register section
+    switch (arch){
+      case SYMS_Arch_X86:
+      {
+        switch (frame_ptr_reg){
+          case SYMS_CvEncodedFramePtrReg_StackPtr:
+          {
+            // TODO(allen): support SYMS_CvAllReg_VFRAME
+          }break;
+          case SYMS_CvEncodedFramePtrReg_FramePtr:
+          {
+            result = syms_pdb_reg_section_from_x86_reg(SYMS_CvRegx86_EBP);
+          }break;
+          case SYMS_CvEncodedFramePtrReg_BasePtr:
+          {
+            result = syms_pdb_reg_section_from_x86_reg(SYMS_CvRegx86_EBX);
+          }break;
+        }
+      }break;
+      case SYMS_Arch_X64:
+      {
+        switch (frame_ptr_reg){
+          case SYMS_CvEncodedFramePtrReg_StackPtr:
+          {
+            result = syms_pdb_reg_section_from_x64_reg(SYMS_CvRegx64_RSP);
+          }break;
+          case SYMS_CvEncodedFramePtrReg_FramePtr:
+          {
+            result = syms_pdb_reg_section_from_x64_reg(SYMS_CvRegx64_RBP);
+          }break;
+          case SYMS_CvEncodedFramePtrReg_BasePtr:
+          {
+            result = syms_pdb_reg_section_from_x64_reg(SYMS_CvRegx64_R13);
+          }break;
+        }
+      }break;
+    }
+  }
+  
+  //- fallback
+  else{
+    switch (arch){
+      case SYMS_Arch_X86:
+      {
+        result = syms_reg_section_X86[SYMS_RegX86Code_ebp];
+      }break;
+      case SYMS_Arch_X64:
+      {
+        result = syms_reg_section_X64[SYMS_RegX64Code_rbp];
+      }break;
+    }
+  }
+  
+  return(result);
+}
 
 // main api
 
@@ -3455,6 +3691,354 @@ syms_pdb_voff_from_var_sid(SYMS_String8 data, SYMS_PdbDbgAccel *dbg, SYMS_PdbUni
   return(result);
 }
 
+SYMS_API SYMS_Location
+syms_pdb_location_from_var_sid(SYMS_Arena *arena, SYMS_String8 data, SYMS_PdbDbgAccel *dbg,
+                               SYMS_PdbUnitAccel *unit, SYMS_SymbolID sid){
+  SYMS_Location result = {0};
+  if (!unit->leaf_set){
+    // setup accel
+    SYMS_MsfAccel *msf = dbg->msf;
+    
+    // read id
+    SYMS_MsfRange range = syms_msf_range_from_sn(msf, unit->sn);
+    SYMS_PdbStub *stub = syms_pdb_stub_from_unit_off(unit, SYMS_ID_u32_1(sid));
+    
+    // parse symbol
+    if (stub != 0){
+      SYMS_CvElement element = syms_cv_element(data, msf, range, stub->off);
+      
+      switch (element.kind){
+        case SYMS_CvSymKind_LDATA32:
+        case SYMS_CvSymKind_GDATA32:
+        {
+          SYMS_CvData32 data32 = {0};
+          syms_msf_read_struct_in_range(data, msf, element.range, 0, &data32);
+          SYMS_CoffSection section = syms_pdb_coff_section(data, dbg, data32.sec);
+          SYMS_U32 voff = section.virt_off + data32.sec_off;
+          
+          //- build the location info for addr:`module base + off`
+          SYMS_EvalOpList list = {0};
+          syms_op_push(arena, &list, SYMS_EvalOp_ModuleOff, syms_op_params(voff));
+          
+          result.op_list = list;
+          result.mode = SYMS_EvalMode_Address;
+        }break;
+        
+        case SYMS_CvSymKind_REGREL32:
+        {
+          //- extract info
+          SYMS_CvRegrel32 regrel32 = {0};
+          syms_msf_read_struct_in_range(data, msf, element.range, 0, &regrel32);
+          
+          SYMS_CvReg cv_reg = regrel32.reg;
+          SYMS_U32 off = regrel32.reg_off;
+          SYMS_Arch arch = syms_pdb_arch_from_dbg(dbg);
+          // TODO(allen): report unimplemented architecture, unimplemented register convertion
+          SYMS_RegSection sec = syms_pdb_reg_section_from_arch_reg(arch, cv_reg);
+          
+          //- build the location info for addr:`reg + off`
+          SYMS_EvalOpList list = {0};
+          syms_op_push(arena, &list, SYMS_EvalOp_RegRead, syms_op_params_2u16(sec.off, sec.size));
+          if (off != 0){
+            syms_op_encode_u(arena, &list, off);
+            syms_op_push(arena, &list, SYMS_EvalOp_Add, syms_op_params(SYMS_EvalTypeGroup_U));
+          }
+          
+          result.op_list = list;
+          result.mode = SYMS_EvalMode_Address;
+        }break;
+        
+        case SYMS_CvSymKind_GTHREAD32:
+        case SYMS_CvSymKind_LTHREAD32:
+        {
+          SYMS_CvThread32 thread32 = {0};
+          syms_msf_read_struct_in_range(data, msf, element.range, 0, &thread32);
+          
+          //- build the location info for addr:`TLS base + off`
+          SYMS_EvalOpList list = {0};
+          if (thread32.tls_off <= 0xFFFF){
+            syms_op_push(arena, &list, SYMS_EvalOp_TLSOff, syms_op_params(thread32.tls_off));
+          }
+          else{
+            syms_op_push(arena, &list, SYMS_EvalOp_TLSOff, syms_op_params(0xFFFF));
+            syms_op_encode_u(arena, &list, (thread32.tls_off - 0xFFFF));
+            syms_op_push(arena, &list, SYMS_EvalOp_Add, syms_op_params(SYMS_EvalTypeGroup_U));
+          }
+          
+          result.op_list = list;
+          result.mode = SYMS_EvalMode_Address;
+        }break;
+      }
+    }
+  }
+  return(result);
+}
+
+SYMS_API SYMS_LocRangeArray
+syms_pdb_location_ranges_from_var_sid(SYMS_Arena *arena, SYMS_String8 data, SYMS_PdbDbgAccel *dbg,
+                                      SYMS_PdbUnitAccel *unit, SYMS_SymbolID sid){
+  SYMS_LocRangeArray result = {0};
+  if (!unit->leaf_set){
+    //- setup accel
+    SYMS_MsfAccel *msf = dbg->msf;
+    
+    //- read id
+    SYMS_MsfRange range = syms_msf_range_from_sn(msf, unit->sn);
+    SYMS_PdbStub *stub = syms_pdb_stub_from_unit_off(unit, SYMS_ID_u32_1(sid));
+    
+    //- parse symbol
+    if (stub != 0){
+      SYMS_CvElement element = syms_cv_element(data, msf, range, stub->off);
+      if (element.kind == SYMS_CvSymKind_LOCAL){
+        SYMS_ArenaTemp scratch = syms_get_scratch(&arena, 1);
+        
+        //- gather ranges
+        SYMS_LocRangeList list = {0};
+        for (SYMS_PdbStub *child = stub->first;
+             child != 0;
+             child = child->sibling_next){
+          SYMS_CvElement child_element = syms_cv_element(data, msf, range, child->off);
+          SYMS_LocID loc_id = SYMS_ID_u32_u32(SYMS_PdbSymbolIDKind_Off, child->off);
+          
+          //- determine handling path
+          SYMS_B32 do_standard_range = syms_true;
+          SYMS_U64 range_off = 0;
+          switch (child_element.kind){
+            case SYMS_CvSymKind_DEFRANGE_2005:
+            case SYMS_CvSymKind_DEFRANGE2_2005:
+            {
+              // TODO(allen): Investigate these cases further
+              do_standard_range = syms_false;
+            }break;
+            
+            case SYMS_CvSymKind_DEFRANGE:
+            {
+              range_off = SYMS_MEMBER_OFFSET(SYMS_CvDefrange, range);
+            }break;
+            
+            case SYMS_CvSymKind_DEFRANGE_SUBFIELD:
+            {
+              range_off = SYMS_MEMBER_OFFSET(SYMS_CvDefrangeSubfield, range);
+            }break;
+            
+            case SYMS_CvSymKind_DEFRANGE_REGISTER:
+            {
+              range_off = SYMS_MEMBER_OFFSET(SYMS_CvDefrangeRegister, range);
+            }break;
+            
+            case SYMS_CvSymKind_DEFRANGE_FRAMEPOINTER_REL:
+            {
+              range_off = SYMS_MEMBER_OFFSET(SYMS_CvDefrangeFramepointerRel, range);
+            }break;
+            
+            case SYMS_CvSymKind_DEFRANGE_SUBFIELD_REGISTER:
+            {
+              range_off = SYMS_MEMBER_OFFSET(SYMS_CvDefrangeSubfieldRegister, range);
+            }break;
+            
+            case SYMS_CvSymKind_DEFRANGE_FRAMEPOINTER_REL_FULL_SCOPE:
+            {
+              do_standard_range = syms_false;
+              
+              //- emit a single range that covers all U64
+              SYMS_LocRangeNode *node = syms_push_array(scratch.arena, SYMS_LocRangeNode, 1);
+              SYMS_QueuePush(list.first, list.last, node);
+              list.count += 1;
+              node->loc_range.vrange.min = 0;
+              node->loc_range.vrange.max = SYMS_U64_MAX;
+              node->loc_range.loc_id = loc_id;
+            }break;
+            
+            case SYMS_CvSymKind_DEFRANGE_REGISTER_REL:
+            {
+              range_off = SYMS_MEMBER_OFFSET(SYMS_CvDefrangeRegisterRel, range);
+            }break;
+          }
+          
+          //- apply range handling
+          if (do_standard_range){
+            SYMS_CvLvarAddrRange addr_range = {0};
+            if (syms_msf_read_struct_in_range(data, msf, child_element.range, range_off, &addr_range)){
+              // setup main range
+              SYMS_CoffSection section = syms_pdb_coff_section(data, dbg, addr_range.sec);
+              SYMS_U64 vaddr_base = section.virt_off + addr_range.off;
+              SYMS_U64 vaddr_last = vaddr_base + addr_range.len;
+              
+              // setup gaps
+              SYMS_U64 gap_off = range_off + sizeof(addr_range);
+              SYMS_U64 gap_count = (child_element.range.size - gap_off)/sizeof(SYMS_CvLvarAddrGap);
+              SYMS_U64 gap_opl = gap_off + gap_count*sizeof(SYMS_CvLvarAddrGap);
+              
+              // emit loop
+              SYMS_U64 vaddr_cursor = vaddr_base;
+              for (SYMS_U64 gap_cursor = gap_off;
+                   gap_cursor < gap_opl;
+                   gap_cursor += sizeof(SYMS_CvLvarAddrGap)){
+                SYMS_CvLvarAddrGap gap = {0};
+                syms_msf_read_struct_in_range(data, msf, child_element.range, gap_off, &gap);
+                
+                SYMS_U64 gap_first = vaddr_base + gap.off;
+                
+                // emit range
+                {
+                  SYMS_LocRangeNode *node = syms_push_array(scratch.arena, SYMS_LocRangeNode, 1);
+                  SYMS_QueuePush(list.first, list.last, node);
+                  list.count += 1;
+                  node->loc_range.vrange.min = vaddr_cursor;
+                  node->loc_range.vrange.max = gap_first;
+                  node->loc_range.loc_id = loc_id;
+                }
+                
+                // advance vaddr cursor
+                vaddr_cursor = gap_first + gap.len;
+              }
+              
+              // emit range
+              {
+                SYMS_LocRangeNode *node = syms_push_array(scratch.arena, SYMS_LocRangeNode, 1);
+                SYMS_QueuePush(list.first, list.last, node);
+                list.count += 1;
+                node->loc_range.vrange.min = vaddr_cursor;
+                node->loc_range.vrange.max = vaddr_last;
+                node->loc_range.loc_id = loc_id;
+              }
+            }
+          }
+        }
+        
+        //- flatten list
+        SYMS_LocRange *loc_ranges = syms_push_array(arena, SYMS_LocRange, list.count);
+        SYMS_LocRange *loc_range_ptr = loc_ranges;
+        for (SYMS_LocRangeNode *node = list.first;
+             node != 0;
+             node = node->next, loc_range_ptr += 1){
+          syms_memmove(loc_range_ptr, &node->loc_range, sizeof(node->loc_range));
+        }
+        
+        //- fill result
+        result.loc_ranges = loc_ranges;
+        result.count = list.count;
+        
+        syms_release_scratch(scratch);
+      }
+    }
+  }
+  return(result);
+}
+
+SYMS_API SYMS_Location
+syms_pdb_location_from_id(SYMS_Arena *arena, SYMS_String8 data, SYMS_PdbDbgAccel *dbg,
+                          SYMS_PdbUnitAccel *unit, SYMS_LocID loc_id){
+  SYMS_Location result = {0};
+  if (!unit->leaf_set){
+    // setup accel
+    SYMS_MsfAccel *msf = dbg->msf;
+    
+    // read id
+    SYMS_MsfRange range = syms_msf_range_from_sn(msf, unit->sn);
+    SYMS_PdbStub *stub = syms_pdb_stub_from_unit_off(unit, SYMS_ID_u32_1(loc_id));
+    
+    // parse symbol
+    if (stub != 0){
+      SYMS_CvElement element = syms_cv_element(data, msf, range, stub->off);
+      
+      switch (element.kind){
+        {
+          case SYMS_CvSymKind_DEFRANGE_2005:
+          case SYMS_CvSymKind_DEFRANGE2_2005:
+          {
+            // TODO(allen): Investigate these cases further
+          }break;
+          
+          case SYMS_CvSymKind_DEFRANGE:
+          {
+            // TODO(allen): SYMS_CvDefrange 'program' - don't know how to interpret
+          }break;
+          
+          case SYMS_CvSymKind_DEFRANGE_SUBFIELD:
+          {
+            // TODO(allen): SYMS_CvDefrangeSubfield 'program' - don't know how to interpret
+            // TODO(allen): SYMS_CvDefrangeSubfield 'off_in_parent' - not a supported concept
+          }break;
+          
+          case SYMS_CvSymKind_DEFRANGE_REGISTER:
+          {
+            // extract info
+            SYMS_CvDefrangeRegister defrange_register = {0};
+            syms_msf_read_struct_in_range(data, msf, element.range, 0, &defrange_register);
+            
+            SYMS_CvReg cv_reg = defrange_register.reg;
+            SYMS_Arch arch = syms_pdb_arch_from_dbg(dbg);
+            SYMS_RegSection sec = syms_pdb_reg_section_from_arch_reg(arch, cv_reg);
+            
+            // location info ops
+            SYMS_EvalOpList list = {0};
+            syms_op_encode_reg_section(arena, &list, sec);
+            
+            // fill result
+            result.op_list = list;
+            result.mode = SYMS_EvalMode_Register;
+          }break;
+          
+          case SYMS_CvSymKind_DEFRANGE_FRAMEPOINTER_REL:
+          case SYMS_CvSymKind_DEFRANGE_FRAMEPOINTER_REL_FULL_SCOPE:
+          {
+            // extract framepointer register
+            SYMS_RegSection sec = syms_pdb_reg_section_from_framepointer(data, dbg, range, stub);
+            
+            // extract offset
+            SYMS_U32 offset = 0;
+            syms_msf_read_struct_in_range(data, msf, element.range, 0, &offset);
+            
+            // location info ops
+            SYMS_EvalOpList list = {0};
+            syms_op_push(arena, &list, SYMS_EvalOp_RegRead, syms_op_params_2u16(sec.off, sec.size));
+            if (offset != 0){
+              syms_op_encode_u(arena, &list, offset);
+              syms_op_push(arena, &list, SYMS_EvalOp_Add, syms_op_params(SYMS_EvalTypeGroup_U));
+            }
+            
+            // fill result
+            result.op_list = list;
+            result.mode = SYMS_EvalMode_Address;
+          }break;
+          
+          case SYMS_CvSymKind_DEFRANGE_SUBFIELD_REGISTER:
+          {
+            // TODO(allen): SYMS_CvDefrangeSubfieldRegister 'off_in_parent' - not a supported concept
+          }break;
+          
+          case SYMS_CvSymKind_DEFRANGE_REGISTER_REL:
+          {
+            // extract info
+            SYMS_CvDefrangeRegisterRel defrange_register_rel = {0};
+            syms_msf_read_struct_in_range(data, msf, element.range, 0, &defrange_register_rel);
+            
+            SYMS_CvReg cv_reg = defrange_register_rel.reg;
+            SYMS_U32 off = defrange_register_rel.reg_off;
+            SYMS_Arch arch = syms_pdb_arch_from_dbg(dbg);
+            SYMS_RegSection sec = syms_pdb_reg_section_from_arch_reg(arch, cv_reg);
+            
+            // location info ops
+            SYMS_EvalOpList list = {0};
+            syms_op_push(arena, &list, SYMS_EvalOp_RegRead, syms_op_params_2u16(sec.off, sec.size));
+            if (off != 0){
+              syms_op_encode_u(arena, &list, off);
+              syms_op_push(arena, &list, SYMS_EvalOp_Add, syms_op_params(SYMS_EvalTypeGroup_U));
+            }
+            
+            // fill result
+            result.op_list = list;
+            result.mode = SYMS_EvalMode_Address;;
+          }break;
+        }
+      }
+      
+    }
+  }
+  return(result);
+}
+
 
 ////////////////////////////////
 //~ NOTE(allen): PDB Member Info
@@ -3581,6 +4165,7 @@ syms_pdb_mems_accel_from_sid(SYMS_Arena *arena, SYMS_String8 data, SYMS_PdbDbgAc
             ti = struct_.field;
           }
         }break;
+        
         case SYMS_CvLeaf_UNION:
         {
           SYMS_CvLeafUnion union_ = {0};
@@ -3588,6 +4173,15 @@ syms_pdb_mems_accel_from_sid(SYMS_Arena *arena, SYMS_String8 data, SYMS_PdbDbgAc
             ti = union_.field;
           }
         }break;
+        
+        case SYMS_CvLeaf_CLASSPTR:
+        case SYMS_CvLeaf_CLASSPTR2:
+        {
+          SYMS_CvLeafClassPtr class_ptr = {0};
+          if (syms_msf_read_struct_in_range(data, msf, element.range, 0, &class_ptr)){
+            ti = class_ptr.field;
+          }
+        }break;
       }
     }
     
@@ -4022,13 +4616,8 @@ syms_pdb_enum_info_array_from_sid(SYMS_Arena *arena, SYMS_String8 data, SYMS_Pdb
         SYMS_CvLeaf lf_kind = 0;
         syms_msf_read_struct_in_range(data, msf, range, lf_off, &lf_kind);
         if (lf_kind == SYMS_CvLeaf_ENUMERATE){
-          SYMS_U32 num_off = lf_off + sizeof(SYMS_CvLeafEnumerate);
-          SYMS_CvNumeric num = {SYMS_TypeKind_Null};
-          SYMS_U32 num_size = syms_cv_read_numeric(data, msf, range, num_off, &num);
-          SYMS_U32 size_clamped = SYMS_ClampTop(num_size, sizeof(enum_info_ptr->val));
-          
           enum_info_ptr->name = syms_push_string_copy(arena, field_stub->name);
-          syms_memmove(&enum_info_ptr->val, num.data, size_clamped);
+          syms_memmove(&enum_info_ptr->val, &field_stub->num, sizeof(field_stub->num));
           enum_info_ptr += 1;
         }
       }
@@ -4090,8 +4679,8 @@ syms_pdb_sig_info_from_handle(SYMS_Arena *arena, SYMS_String8 data, SYMS_PdbDbgA
 }
 
 SYMS_API SYMS_U64RangeArray
-syms_pdb_proc_vranges_from_sid(SYMS_Arena *arena, SYMS_String8 data, SYMS_PdbDbgAccel *dbg,
-                               SYMS_PdbUnitAccel *unit, SYMS_SymbolID sid){
+syms_pdb_scope_vranges_from_sid(SYMS_Arena *arena, SYMS_String8 data, SYMS_PdbDbgAccel *dbg,
+                                SYMS_PdbUnitAccel *unit, SYMS_SymbolID sid){
   SYMS_U64RangeArray result = {0};
   if (!unit->leaf_set){
     // setup accels
@@ -4123,6 +4712,22 @@ syms_pdb_proc_vranges_from_sid(SYMS_Arena *arena, SYMS_String8 data, SYMS_PdbDbg
             result.ranges[0] = range;
           }
         }break;
+        
+        case SYMS_CvSymKind_BLOCK32:
+        {
+          SYMS_CvBlock32 block32 = {0};
+          if (syms_msf_read_struct_in_range(data, msf, element.range, 0, &block32)){
+            SYMS_CoffSection section = syms_pdb_coff_section(data, dbg, block32.sec);
+            
+            SYMS_U64Range range = {0};
+            range.min = section.virt_off + block32.off;
+            range.max = range.min + block32.len;
+            
+            result.count = 1;
+            result.ranges = syms_push_array(arena, SYMS_U64Range, 1);
+            result.ranges[0] = range;
+          }
+        }break;
       }
     }
   }
diff --git a/Engine/Source/Runtime/SymsLib/syms/core/pdb/syms_pdb_parser.h b/Engine/Source/Runtime/SymsLib/syms/core/pdb/syms_pdb_parser.h
--- a/Engine/Source/Runtime/SymsLib/syms/core/pdb/syms_pdb_parser.h
+++ b/Engine/Source/Runtime/SymsLib/syms/core/pdb/syms_pdb_parser.h
@@ -35,7 +35,6 @@ typedef enum SYMS_PdbPseudoUnit{
   SYMS_PdbPseudoUnit_Null,
   SYMS_PdbPseudoUnit_SYM,
   SYMS_PdbPseudoUnit_TPI,
-  SYMS_PdbPseudoUnit_LAST_PSEUDO_UNIT = SYMS_PdbPseudoUnit_TPI,
   SYMS_PdbPseudoUnit_COUNT = SYMS_PdbPseudoUnit_TPI,
   SYMS_PdbPseudoUnit_FIRST_COMP_UNIT,
 } SYMS_PdbPseudoUnit;
@@ -118,6 +117,7 @@ typedef struct SYMS_PdbUnitSetAccel{
 
 typedef struct SYMS_PdbStub{
   struct SYMS_PdbStub *bucket_next;
+  struct SYMS_PdbStub *parent;
   struct SYMS_PdbStub *sibling_next;
   struct SYMS_PdbStub *first;
   struct SYMS_PdbStub *last;
@@ -152,6 +152,8 @@ typedef struct SYMS_PdbUnitAccel{
   SYMS_U64 proc_count;
   SYMS_PdbStub **var_stubs;
   SYMS_U64 var_count;
+  SYMS_PdbStub **tls_var_stubs;
+  SYMS_U64 tls_var_count;
   SYMS_PdbStub **thunk_stubs;
   SYMS_U64 thunk_count;
   SYMS_PdbStub **pub_stubs;
@@ -453,6 +455,10 @@ SYMS_API SYMS_UnitID       syms_pdb_uid_from_accel(SYMS_PdbUnitAccel *unit);
 
 SYMS_API SYMS_SymbolIDArray syms_pdb_proc_sid_array_from_unit(SYMS_Arena *arena, SYMS_PdbUnitAccel *unit);
 SYMS_API SYMS_SymbolIDArray syms_pdb_var_sid_array_from_unit(SYMS_Arena *arena, SYMS_PdbUnitAccel *unit);
+
+SYMS_API SYMS_UnitID       syms_pdb_tls_var_uid_from_dbg(SYMS_PdbDbgAccel *dbg);
+SYMS_API SYMS_SymbolIDArray syms_pdb_tls_var_sid_array_from_unit(SYMS_Arena *arena, SYMS_PdbUnitAccel *unit);
+
 SYMS_API SYMS_SymbolIDArray syms_pdb_type_sid_array_from_unit(SYMS_Arena *arena, SYMS_PdbUnitAccel *unit);
 
 SYMS_API SYMS_U64          syms_pdb_symbol_count_from_unit(SYMS_PdbUnitAccel *unit);
@@ -478,6 +484,13 @@ SYMS_API SYMS_USID   syms_pdb_sym_type_from_var_id(SYMS_String8 data, SYMS_PdbDb
 SYMS_API SYMS_U64    syms_pdb_sym_voff_from_var_sid(SYMS_String8 data, SYMS_PdbDbgAccel *dbg,
                                                     SYMS_PdbUnitAccel *unit, SYMS_SymbolID sid);
 
+SYMS_API SYMS_RegSection syms_pdb_reg_section_from_x86_reg(SYMS_CvReg cv_reg);
+SYMS_API SYMS_RegSection syms_pdb_reg_section_from_x64_reg(SYMS_CvReg cv_reg);
+SYMS_API SYMS_RegSection syms_pdb_reg_section_from_arch_reg(SYMS_Arch arch, SYMS_CvReg cv_reg);
+
+SYMS_API SYMS_RegSection syms_pdb_reg_section_from_framepointer(SYMS_String8 data,  SYMS_PdbDbgAccel *dbg,
+                                                                SYMS_MsfRange range, SYMS_PdbStub *framepointer_stub);
+
 // main api
 SYMS_API SYMS_USID   syms_pdb_type_from_var_id(SYMS_String8 data, SYMS_PdbDbgAccel *dbg, SYMS_PdbUnitAccel *unit,
                                                SYMS_SymbolID id);
@@ -485,6 +498,13 @@ SYMS_API SYMS_USID   syms_pdb_type_from_var_id(SYMS_String8 data, SYMS_PdbDbgAcc
 SYMS_API SYMS_U64    syms_pdb_voff_from_var_sid(SYMS_String8 data, SYMS_PdbDbgAccel *dbg, SYMS_PdbUnitAccel *unit,
                                                 SYMS_SymbolID sid);
 
+SYMS_API SYMS_Location syms_pdb_location_from_var_sid(SYMS_Arena *arena, SYMS_String8 data, SYMS_PdbDbgAccel *dbg,
+                                                      SYMS_PdbUnitAccel *unit, SYMS_SymbolID sid);
+SYMS_API SYMS_LocRangeArray syms_pdb_location_ranges_from_var_sid(SYMS_Arena *arena, SYMS_String8 data, SYMS_PdbDbgAccel *dbg,
+                                                                  SYMS_PdbUnitAccel *unit, SYMS_SymbolID sid);
+SYMS_API SYMS_Location syms_pdb_location_from_id(SYMS_Arena *arena, SYMS_String8 data, SYMS_PdbDbgAccel *dbg,
+                                                 SYMS_PdbUnitAccel *unit, SYMS_LocID loc_id);
+
 
 ////////////////////////////////
 //~ NOTE(allen): PDB Member Info
@@ -525,9 +545,9 @@ SYMS_API SYMS_SigInfo      syms_pdb_sig_info_from_handle(SYMS_Arena *arena, SYMS
                                                          SYMS_PdbDbgAccel *dbg, SYMS_PdbUnitAccel *unit,
                                                          SYMS_SigHandle handle);
 
-SYMS_API SYMS_U64RangeArray syms_pdb_proc_vranges_from_sid(SYMS_Arena *arena, SYMS_String8 data,
-                                                           SYMS_PdbDbgAccel *dbg, SYMS_PdbUnitAccel *unit,
-                                                           SYMS_SymbolID sid);
+SYMS_API SYMS_U64RangeArray syms_pdb_scope_vranges_from_sid(SYMS_Arena *arena, SYMS_String8 data,
+                                                            SYMS_PdbDbgAccel *dbg, SYMS_PdbUnitAccel *unit,
+                                                            SYMS_SymbolID sid);
 SYMS_API SYMS_SymbolIDArray syms_pdb_scope_children_from_sid(SYMS_Arena *arena, SYMS_String8 data,
                                                              SYMS_PdbDbgAccel *dbg, SYMS_PdbUnitAccel *unit,
                                                              SYMS_SymbolID id);
@@ -560,7 +580,6 @@ SYMS_API SYMS_Line* syms_pdb_loose_push_sequence(SYMS_Arena *arena, SYMS_PdbLine
                                                  SYMS_U64 line_count);
 
 // main api
-
 SYMS_API SYMS_String8 syms_pdb_file_name_from_id(SYMS_Arena *arena, SYMS_String8 data, SYMS_PdbDbgAccel *dbg,
                                                  SYMS_PdbUnitSetAccel *unit_set, SYMS_UnitID uid,
                                                  SYMS_FileID id);
diff --git a/Engine/Source/Runtime/SymsLib/syms/core/pe/syms_pe_parser.c b/Engine/Source/Runtime/SymsLib/syms/core/pe/syms_pe_parser.c
--- a/Engine/Source/Runtime/SymsLib/syms/core/pe/syms_pe_parser.c
+++ b/Engine/Source/Runtime/SymsLib/syms/core/pe/syms_pe_parser.c
@@ -1,4 +1,5 @@
 // Copyright Epic Games, Inc. All Rights Reserved.
+
 #ifndef SYMS_PE_PARSER_C
 #define SYMS_PE_PARSER_C
 
@@ -61,7 +62,8 @@ syms_pe_bin_accel_from_file(SYMS_Arena *arena, SYMS_String8 data, SYMS_PeFileAcc
     SYMS_U64 clamped_sec_count = (sec_array_opl - sec_array_off)/sizeof(SYMS_CoffSection);
     SYMS_CoffSection *sections = (SYMS_CoffSection*)((SYMS_U8*)base + sec_array_off);
     
-    // get data directory from optional
+    // read ptional header
+    SYMS_U64 image_base = 0;
     SYMS_U32 data_dir_count = 0;
     SYMS_U64Range *data_dirs_file = 0;
     SYMS_U64Range *data_dirs_virt = 0;
@@ -78,6 +80,9 @@ syms_pe_bin_accel_from_file(SYMS_Arena *arena, SYMS_String8 data, SYMS_PeFileAcc
         {
           SYMS_PeOptionalPe32 pe_optional = {0};
           syms_based_range_read_struct(base, optional_range, 0, &pe_optional);
+
+          image_base = pe_optional.image_base;
+
           reported_data_dir_offset = sizeof(pe_optional);
           reported_data_dir_count = pe_optional.data_dir_count;
         }break;
@@ -85,6 +90,9 @@ syms_pe_bin_accel_from_file(SYMS_Arena *arena, SYMS_String8 data, SYMS_PeFileAcc
         {
           SYMS_PeOptionalPe32Plus pe_optional = {0};
           syms_based_range_read_struct(base, optional_range, 0, &pe_optional);
+
+          image_base = pe_optional.image_base;
+
           reported_data_dir_offset = sizeof(pe_optional);
           reported_data_dir_count = pe_optional.data_dir_count;
         }break;
@@ -151,6 +159,7 @@ syms_pe_bin_accel_from_file(SYMS_Arena *arena, SYMS_String8 data, SYMS_PeFileAcc
     // fill result
     result = syms_push_array(arena, SYMS_PeBinAccel, 1);
     result->format = SYMS_FileFormat_PE;
+    result->image_base = image_base;
     result->section_array_off = sec_array_off;
     result->section_count = clamped_sec_count;
     result->dbg_path_off = dbg_path_off;
@@ -234,7 +243,7 @@ syms_pe_binary_search_intel_pdata(SYMS_String8 data, SYMS_PeBinAccel *bin, SYMS_
   SYMS_ASSERT(bin->arch == SYMS_Arch_X86 || bin->arch == SYMS_Arch_X64);
   SYMS_U64 pdata_off = bin->data_dirs_file[SYMS_PeDataDirectoryIndex_EXCEPTIONS].min;
   SYMS_U64 pdata_count = syms_u64_range_size(bin->data_dirs_file[SYMS_PeDataDirectoryIndex_EXCEPTIONS]) / sizeof(SYMS_PeIntelPdata);
-
+  
   SYMS_U64 result = 0;
   // check if this bin includes a pdata array
   if (pdata_count > 0){
@@ -367,6 +376,7 @@ syms_pe_imports_from_bin(SYMS_Arena *arena, SYMS_String8 data, SYMS_PeBinAccel *
   //- rjf: grab prerequisites
   void *base = (void*)data.str;
   SYMS_Arch arch = syms_pe_arch_from_bin(bin);
+  (void)arch;
   SYMS_U64 arch_bit_size = syms_address_size_from_arch(bin->arch);
   SYMS_U64 arch_byte_size = arch_bit_size / 8;
   
@@ -430,7 +440,7 @@ syms_pe_imports_from_bin(SYMS_Arena *arena, SYMS_String8 data, SYMS_PeBinAccel *
         dll_node->import_count += 1;
         
         // rjf: grab name/ordinal
-        SYMS_B32 is_ordinal = (arch_bit_size == 64) ? lookup_entry & 0x8000000000000000 : lookup_entry & 0x80000000;
+        SYMS_B32 is_ordinal = !!((arch_bit_size == 64) ? lookup_entry & 0x8000000000000000 : lookup_entry & 0x80000000);
         if(is_ordinal)
         {
           import_node->ordinal = lookup_entry & 0xFFFF;
diff --git a/Engine/Source/Runtime/SymsLib/syms/core/regs/syms_regs.c b/Engine/Source/Runtime/SymsLib/syms/core/regs/syms_regs.c
--- a/Engine/Source/Runtime/SymsLib/syms/core/regs/syms_regs.c
+++ b/Engine/Source/Runtime/SymsLib/syms/core/regs/syms_regs.c
@@ -1,19 +1,6 @@
 // Copyright Epic Games, Inc. All Rights Reserved.
+
 #ifndef SYMS_REGS_C
 #define SYMS_REGS_C
 
-////////////////////////////////
-// NOTE(allen): Basic Register Helpers
-
-SYMS_API SYMS_String8
-syms_reg_from_metadata_id(SYMS_RegMetadata *metadata, SYMS_U64 count, void *reg_file, SYMS_RegID reg_id){
-  SYMS_String8 result = {0};
-  if (reg_id <= count){
-    SYMS_RegMetadata *m = metadata + reg_id;
-    result.str = (SYMS_U8*)reg_file + m->offset;
-    result.size = m->size;
-  }
-  return(result);
-}
-
 #endif // SYMS_REGS_C
diff --git a/Engine/Source/Runtime/SymsLib/syms/core/regs/syms_regs.h b/Engine/Source/Runtime/SymsLib/syms/core/regs/syms_regs.h
--- a/Engine/Source/Runtime/SymsLib/syms/core/regs/syms_regs.h
+++ b/Engine/Source/Runtime/SymsLib/syms/core/regs/syms_regs.h
@@ -1,4 +1,5 @@
 // Copyright Epic Games, Inc. All Rights Reserved.
+
 #ifndef SYMS_REGS_H
 #define SYMS_REGS_H
 
@@ -46,23 +47,15 @@ typedef union SYMS_Reg256{
 } SYMS_Reg256;
 
 typedef struct SYMS_RegSection{
-  SYMS_U32 offset;
-  SYMS_U32 size;
+  // TODO(allen): naming pass byte_*; double check these are encoded as byte sizes
+  SYMS_U16 off;
+  SYMS_U16 size;
 } SYMS_RegSection;
 
-typedef struct SYMS_RegMetadata{
-  SYMS_U32 offset;
-  SYMS_U32 size;
-  SYMS_String8 name;
-  SYMS_U32 reg_class;
-  SYMS_U32 alias_to;
-} SYMS_RegMetadata;
-
-////////////////////////////////
-// NOTE(allen): Basic Register Helpers
-
-SYMS_API SYMS_String8 syms_reg_from_metadata_id(SYMS_RegMetadata *metadata, SYMS_U64 count,
-                                                void *reg_file, SYMS_RegID reg_id);
-#define syms_reg_from_arch_id(AR,f,i) syms_reg_from_metadata_id(reg_metadata_##AR,SYMS_Reg##AR##Code_COUNT,(f),(i))
+typedef struct SYMS_RegSlice{
+  SYMS_U16 reg_id;
+  SYMS_U8 byte_off;
+  SYMS_U8 byte_size;
+} SYMS_RegSlice;
 
 #endif // SYMS_REGS_H
diff --git a/Engine/Source/Runtime/SymsLib/syms/core/syms_debug_info.c b/Engine/Source/Runtime/SymsLib/syms/core/syms_debug_info.c
--- a/Engine/Source/Runtime/SymsLib/syms/core/syms_debug_info.c
+++ b/Engine/Source/Runtime/SymsLib/syms/core/syms_debug_info.c
@@ -1,4 +1,5 @@
 // Copyright Epic Games, Inc. All Rights Reserved.
+
 #ifndef SYMS_DEBUG_INFO_C
 #define SYMS_DEBUG_INFO_C
 
@@ -57,13 +58,20 @@ syms_type_kind_main_from_fwd(SYMS_TypeKind kind){
 
 SYMS_API SYMS_B32
 syms_type_kind_is_basic(SYMS_TypeKind kind){
-  SYMS_B32 result = (SYMS_TypeKind_Int8 <= kind && kind <= SYMS_TypeKind_Bool);
+  SYMS_B32 result = (SYMS_TypeKind_Int8 <= kind && kind <= SYMS_TypeKind_Void);
+  return(result);
+}
+
+SYMS_API SYMS_B32
+syms_type_kind_is_basic_or_enum(SYMS_TypeKind kind){
+  SYMS_B32 result = ((SYMS_TypeKind_Int8 <= kind && kind <= SYMS_TypeKind_Void) ||
+                     kind == SYMS_TypeKind_Enum || kind == SYMS_TypeKind_ForwardEnum);
   return(result);
 }
 
 SYMS_API SYMS_B32
 syms_type_kind_is_integer(SYMS_TypeKind kind){
-  SYMS_B32 result = (SYMS_TypeKind_Int8 <= kind && kind <= SYMS_TypeKind_UInt128);
+  SYMS_B32 result = (SYMS_TypeKind_Int8 <= kind && kind <= SYMS_TypeKind_Bool);
   return(result);
 }
 
@@ -85,6 +93,12 @@ syms_type_kind_is_user_defined(SYMS_TypeKind kind){
   return(result);
 }
 
+SYMS_API SYMS_B32
+syms_type_kind_is_forward(SYMS_TypeKind kind){
+  SYMS_B32 result = (SYMS_TypeKind_ForwardStruct <= kind && kind <= SYMS_TypeKind_ForwardEnum);
+  return(result);
+}
+
 SYMS_API SYMS_SymbolIDArray
 syms_sid_array_from_list(SYMS_Arena *arena, SYMS_SymbolIDList *list){
   SYMS_SymbolIDArray result = {0};
diff --git a/Engine/Source/Runtime/SymsLib/syms/core/syms_debug_info.h b/Engine/Source/Runtime/SymsLib/syms/core/syms_debug_info.h
--- a/Engine/Source/Runtime/SymsLib/syms/core/syms_debug_info.h
+++ b/Engine/Source/Runtime/SymsLib/syms/core/syms_debug_info.h
@@ -240,6 +240,10 @@ typedef struct SYMS_USIDList{
   SYMS_U64 count;
 } SYMS_USIDList;
 
+typedef struct SYMS_USIDArray{
+  SYMS_USID *usid;
+  SYMS_U64 count;
+} SYMS_USIDArray;
 
 typedef struct SYMS_SigHandle{
   SYMS_U64 v;
@@ -347,6 +351,37 @@ typedef struct SYMS_EnumInfoArray{
   SYMS_U64 count;
 } SYMS_EnumInfoArray;
 
+////////////////////////////////
+// NOTE(allen): Location Types
+
+typedef SYMS_U64 SYMS_LocID;
+
+typedef struct SYMS_LocRange{
+  SYMS_U64Range vrange;
+  SYMS_LocID loc_id;
+} SYMS_LocRange;
+
+typedef struct SYMS_LocRangeArray{
+  SYMS_LocRange *loc_ranges;
+  SYMS_U64 count;
+} SYMS_LocRangeArray;
+
+typedef struct SYMS_LocRangeNode{
+  struct SYMS_LocRangeNode *next;
+  SYMS_LocRange loc_range;
+} SYMS_LocRangeNode;
+
+typedef struct SYMS_LocRangeList{
+  SYMS_LocRangeNode *first;
+  SYMS_LocRangeNode *last;
+  SYMS_U64 count;
+} SYMS_LocRangeList;
+
+typedef enum SYMS_ProcLoc{
+  SYMS_ProcLoc_FrameBase,
+  SYMS_ProcLoc_ReturnAddress,
+} SYMS_ProcLoc;
+
 ////////////////////////////////
 // NOTE(allen): Nil For Accelerators
 
@@ -365,10 +400,12 @@ SYMS_API SYMS_USID    syms_make_usid(SYMS_UnitID uid, SYMS_SymbolID sid);
 SYMS_API SYMS_TypeKind syms_type_kind_fwd_from_main(SYMS_TypeKind type_kind);
 SYMS_API SYMS_TypeKind syms_type_kind_main_from_fwd(SYMS_TypeKind type_kind);
 SYMS_API SYMS_B32      syms_type_kind_is_basic(SYMS_TypeKind kind);
+SYMS_API SYMS_B32      syms_type_kind_is_basic_or_enum(SYMS_TypeKind kind);
 SYMS_API SYMS_B32      syms_type_kind_is_integer(SYMS_TypeKind kind);
 SYMS_API SYMS_B32      syms_type_kind_is_signed(SYMS_TypeKind kind);
 SYMS_API SYMS_B32      syms_type_kind_is_complex(SYMS_TypeKind kind);
 SYMS_API SYMS_B32      syms_type_kind_is_user_defined(SYMS_TypeKind kind);
+SYMS_API SYMS_B32      syms_type_kind_is_forward(SYMS_TypeKind kind);
 
 SYMS_API SYMS_SymbolIDArray syms_sid_array_from_list(SYMS_Arena *arena, SYMS_SymbolIDList *list);
 
diff --git a/Engine/Source/Runtime/SymsLib/syms/core/syms_eval.h b/Engine/Source/Runtime/SymsLib/syms/core/syms_eval.h
new file mode 100644
--- /dev/null
+++ b/Engine/Source/Runtime/SymsLib/syms/core/syms_eval.h
@@ -0,0 +1,112 @@
+// Copyright Epic Games, Inc. All Rights Reserved.
+/* date = December 8th 2021 11:42 am */
+
+#ifndef SYMS_EVAL_H
+#define SYMS_EVAL_H
+
+////////////////////////////////
+// NOTE(allen): Eval Meta Code
+
+#include "syms/core/generated/syms_meta_eval.h"
+
+////////////////////////////////
+// NOTE(allen): Eval Bytecode Helper Types
+
+typedef enum SYMS_EvalMode{
+  SYMS_EvalMode_Nil,
+  SYMS_EvalMode_Value,
+  SYMS_EvalMode_Address,
+  SYMS_EvalMode_Register,
+} SYMS_EvalMode;
+
+enum{
+  // apears in IR tree
+  SYMS_EvalIRExtKind_Bytecode = SYMS_EvalOp_COUNT,
+  
+  // apears in IR graph
+  SYMS_EvalIRExtKind_Noop,
+  
+  SYMS_EvalIRExtKind_COUNT,
+};
+
+typedef union SYMS_EvalOpParams{
+  SYMS_U8  u8[8];
+  SYMS_U16 u16[4];
+  SYMS_U32 u32[2];
+  SYMS_U64 u64[1];
+  SYMS_String8 data;
+} SYMS_EvalOpParams;
+
+typedef struct SYMS_EvalIRTree{
+  SYMS_U32 op;
+  union{
+    struct{
+      struct SYMS_EvalIRTree *children[3];
+      SYMS_EvalOpParams params;
+    };
+  };
+} SYMS_EvalIRTree;
+
+typedef struct SYMS_EvalOpNode{
+  struct SYMS_EvalOpNode *next;
+  SYMS_U8 op;
+  SYMS_EvalOpParams params;
+} SYMS_EvalOpNode;
+
+typedef struct SYMS_EvalOpList{
+  SYMS_EvalOpNode *first;
+  SYMS_EvalOpNode *last;
+  SYMS_U64 byte_count;
+} SYMS_EvalOpList;
+
+typedef struct SYMS_Location{
+  SYMS_EvalOpList op_list;
+  SYMS_EvalMode mode;
+} SYMS_Location;
+
+////////////////////////////////
+// NOTE(allen): 
+
+SYMS_READ_ONLY SYMS_GLOBAL SYMS_EvalIRTree syms_eval_ir_tree_nil = {0};
+
+////////////////////////////////
+// NOTE(allen): Eval Bytecode Helper Functions
+
+// TODO(allen): naming pass "syms_eval_*"
+
+SYMS_API SYMS_EvalOpParams syms_op_params(SYMS_U64 p);
+SYMS_API SYMS_EvalOpParams syms_op_params_2u8(SYMS_U8 p1, SYMS_U8 p2);
+SYMS_API SYMS_EvalOpParams syms_op_params_2u16(SYMS_U16 p1, SYMS_U16 p2);
+SYMS_API SYMS_EvalOpParams syms_op_params_bytecode(SYMS_String8 bytecode);
+
+SYMS_API void syms_op_push(SYMS_Arena *arena, SYMS_EvalOpList *list, SYMS_EvalOp op,
+                           SYMS_EvalOpParams params);
+SYMS_API void syms_op_push_bytecode(SYMS_Arena *arena, SYMS_EvalOpList *list, SYMS_String8 bytecode);
+SYMS_API void syms_op_list_concat_in_place(SYMS_EvalOpList *left, SYMS_EvalOpList *right);
+
+SYMS_API void syms_op_encode_u(SYMS_Arena *arena, SYMS_EvalOpList *list, SYMS_U64 u);
+SYMS_API void syms_op_encode_s(SYMS_Arena *arena, SYMS_EvalOpList *list, SYMS_S64 s);
+
+SYMS_API void syms_op_encode_reg_section(SYMS_Arena *arena, SYMS_EvalOpList *list, SYMS_RegSection sec);
+SYMS_API void syms_op_encode_reg(SYMS_Arena *arena, SYMS_EvalOpList *list, SYMS_Arch arch,
+                                 SYMS_RegID reg_id);
+
+SYMS_API SYMS_String8 syms_bytecode_from_op_list(SYMS_Arena *arena, SYMS_EvalOpList *list);
+
+SYMS_API SYMS_B32 syms_eval_opcode_type_compatible(SYMS_EvalOp op, SYMS_EvalTypeGroup group);
+
+//- ir tree helpers functions
+
+SYMS_API SYMS_EvalIRTree* syms_ir_tree_const_u(SYMS_Arena *arena, SYMS_U64 v);
+SYMS_API SYMS_EvalIRTree* syms_ir_tree_unary_op(SYMS_Arena *arena, SYMS_EvalOp op,
+                                                SYMS_EvalTypeGroup group, SYMS_EvalIRTree *c);
+SYMS_API SYMS_EvalIRTree* syms_ir_tree_binary_op(SYMS_Arena *arena, SYMS_EvalOp op,
+                                                 SYMS_EvalTypeGroup group,
+                                                 SYMS_EvalIRTree *l, SYMS_EvalIRTree *r);
+SYMS_API SYMS_EvalIRTree* syms_ir_tree_binary_op_u(SYMS_Arena *arena, SYMS_EvalOp op, SYMS_EvalIRTree *l,
+                                                   SYMS_EvalIRTree *r);
+SYMS_API SYMS_EvalIRTree* syms_ir_tree_conditional(SYMS_Arena *arena, SYMS_EvalIRTree *c,
+                                                   SYMS_EvalIRTree *l, SYMS_EvalIRTree *r);
+SYMS_API SYMS_EvalIRTree* syms_ir_tree_bytecode_no_copy(SYMS_Arena *arena, SYMS_String8 bytecode);
+
+#endif //SYMS_EVAL_H
diff --git a/Engine/Source/Runtime/SymsLib/syms/core/syms_parser.c b/Engine/Source/Runtime/SymsLib/syms/core/syms_parser.c
--- a/Engine/Source/Runtime/SymsLib/syms/core/syms_parser.c
+++ b/Engine/Source/Runtime/SymsLib/syms/core/syms_parser.c
@@ -1,4 +1,5 @@
 // Copyright Epic Games, Inc. All Rights Reserved.
+
 #ifndef SYMS_PARSER_C
 #define SYMS_PARSER_C
 
@@ -839,6 +840,76 @@ syms_voff_from_var_sid(SYMS_String8 data, SYMS_DbgAccel *dbg, SYMS_UnitAccel *un
   return(result);
 }
 
+SYMS_API SYMS_Location
+syms_location_from_var_sid(SYMS_Arena *arena, SYMS_String8 data, SYMS_DbgAccel *dbg,
+                           SYMS_UnitAccel *unit, SYMS_SymbolID sid){
+  SYMS_ProfBegin("syms_location_from_var_sid");
+  SYMS_Location result = {0};
+  if (unit->format == dbg->format){
+    switch (unit->format){
+      case SYMS_FileFormat_PDB:
+      {
+        result = syms_pdb_location_from_var_sid(arena, data, &dbg->pdb_accel, &unit->pdb_accel, sid);
+      }break;
+      
+      case SYMS_FileFormat_DWARF:
+      {
+        result = syms_dw_location_from_var_sid(arena, data, &dbg->dw_accel, &unit->dw_accel, sid);
+      }break;
+    }
+  }
+  SYMS_ASSERT_PARANOID(syms_parser_api_invariants());
+  SYMS_ProfEnd();
+  return(result);
+}
+
+SYMS_API SYMS_LocRangeArray
+syms_location_ranges_from_var_sid(SYMS_Arena *arena, SYMS_String8 data, SYMS_DbgAccel *dbg,
+                                  SYMS_UnitAccel *unit, SYMS_SymbolID sid){
+  SYMS_ProfBegin("syms_location_ranges_from_var_sid");
+  SYMS_LocRangeArray result = {0};
+  if (unit->format == dbg->format){
+    switch (unit->format){
+      case SYMS_FileFormat_PDB:
+      {
+        result = syms_pdb_location_ranges_from_var_sid(arena, data, &dbg->pdb_accel,
+                                                       &unit->pdb_accel, sid);
+      }break;
+      
+      case SYMS_FileFormat_DWARF:
+      {
+        result = syms_dw_location_ranges_from_var_sid(arena, data, &dbg->dw_accel, &unit->dw_accel, sid);
+      }break;
+    }
+  }
+  SYMS_ASSERT_PARANOID(syms_parser_api_invariants());
+  SYMS_ProfEnd();
+  return(result);
+}
+
+SYMS_API SYMS_Location
+syms_location_from_id(SYMS_Arena *arena, SYMS_String8 data, SYMS_DbgAccel *dbg,
+                      SYMS_UnitAccel *unit, SYMS_LocID loc_id){
+  SYMS_ProfBegin("syms_location_from_id");
+  SYMS_Location result = {0};
+  if (unit->format == dbg->format){
+    switch (unit->format){
+      case SYMS_FileFormat_PDB:
+      {
+        result = syms_pdb_location_from_id(arena, data, &dbg->pdb_accel, &unit->pdb_accel, loc_id);
+      }break;
+      
+      case SYMS_FileFormat_DWARF:
+      {
+        result = syms_dw_location_from_id(arena, data, &dbg->dw_accel, &unit->dw_accel, loc_id);
+      }break;
+    }
+  }
+  SYMS_ASSERT_PARANOID(syms_parser_api_invariants());
+  SYMS_ProfEnd();
+  return(result);
+}
+
 
 // member info
 
@@ -1071,20 +1142,20 @@ syms_sig_info_from_handle(SYMS_Arena *arena, SYMS_String8 data, SYMS_DbgAccel *d
 }
 
 SYMS_API SYMS_U64RangeArray
-syms_proc_vranges_from_sid(SYMS_Arena *arena, SYMS_String8 data, SYMS_DbgAccel *dbg,
-                           SYMS_UnitAccel *unit, SYMS_SymbolID id){
-  SYMS_ProfBegin("syms_proc_vranges_from_sid");
+syms_scope_vranges_from_sid(SYMS_Arena *arena, SYMS_String8 data, SYMS_DbgAccel *dbg,
+                            SYMS_UnitAccel *unit, SYMS_SymbolID id){
+  SYMS_ProfBegin("syms_scope_vranges_from_sid");
   SYMS_U64RangeArray result = {0};
   if (unit->format == dbg->format){
     switch (unit->format){
       case SYMS_FileFormat_PDB:
       {
-        result = syms_pdb_proc_vranges_from_sid(arena, data, &dbg->pdb_accel, &unit->pdb_accel, id);
+        result = syms_pdb_scope_vranges_from_sid(arena, data, &dbg->pdb_accel, &unit->pdb_accel, id);
       }break;
       
       case SYMS_FileFormat_DWARF:
       {
-        result = syms_dw_proc_vranges_from_sid(arena, data, &dbg->dw_accel, &unit->dw_accel, id);
+        result = syms_dw_scope_vranges_from_sid(arena, data, &dbg->dw_accel, &unit->dw_accel, id);
       }break;
     }
   }
@@ -1116,6 +1187,52 @@ syms_scope_children_from_sid(SYMS_Arena *arena, SYMS_String8 data, SYMS_DbgAccel
   return(result);
 }
 
+SYMS_API SYMS_Location
+syms_location_from_proc_sid(SYMS_Arena *arena, SYMS_String8 data, SYMS_DbgAccel *dbg,
+                            SYMS_UnitAccel *unit, SYMS_SymbolID sid, SYMS_ProcLoc proc_loc){
+  SYMS_ProfBegin("syms_location_from_proc_sid");
+  SYMS_Location result = {0};
+  if (unit->format == dbg->format){
+    switch (unit->format){
+      case SYMS_FileFormat_PDB:
+      {
+        // do nothing
+      }break;
+      
+      case SYMS_FileFormat_DWARF:
+      {
+        result = syms_dw_location_from_proc_sid(arena, data, &dbg->dw_accel, &unit->dw_accel, sid, proc_loc);
+      }break;
+    }
+  }
+  SYMS_ASSERT_PARANOID(syms_parser_api_invariants());
+  SYMS_ProfEnd();
+  return(result);
+}
+
+SYMS_API SYMS_LocRangeArray
+syms_location_ranges_from_proc_sid(SYMS_Arena *arena, SYMS_String8 data, SYMS_DbgAccel *dbg,
+                                   SYMS_UnitAccel *unit, SYMS_SymbolID sid, SYMS_ProcLoc proc_loc){
+  SYMS_ProfBegin("syms_location_ranges_from_proc_sid");
+  SYMS_LocRangeArray result = {0};
+  if (unit->format == dbg->format){
+    switch (unit->format){
+      case SYMS_FileFormat_PDB:
+      {
+        // do nothing
+      }break;
+      
+      case SYMS_FileFormat_DWARF:
+      {
+        result = syms_dw_location_ranges_from_proc_sid(arena, data, &dbg->dw_accel, &unit->dw_accel, sid, proc_loc);
+      }break;
+    }
+  }
+  SYMS_ASSERT_PARANOID(syms_parser_api_invariants());
+  SYMS_ProfEnd();
+  return(result);
+}
+
 SYMS_API SYMS_StrippedInfoArray
 syms_stripped_from_unit(SYMS_Arena *arena, SYMS_String8 data, SYMS_DbgAccel *dbg, SYMS_UnitAccel *unit){
   SYMS_ProfBegin("syms_stripped_from_unit");
@@ -1317,4 +1434,39 @@ syms_symbol_list_from_string(SYMS_Arena *arena, SYMS_String8 data, SYMS_DbgAccel
   return(result);
 }
 
+// thread vars
+
+SYMS_API SYMS_UnitID
+syms_tls_var_uid_from_dbg(SYMS_DbgAccel *dbg){
+  SYMS_UnitID result = 0;
+  switch (dbg->format){
+    case SYMS_FileFormat_PDB:
+    {
+      result = syms_pdb_tls_var_uid_from_dbg(&dbg->pdb_accel);
+    }break;
+  }
+  return(result);
+}
+
+SYMS_API SYMS_SymbolIDArray
+syms_tls_var_sid_array_from_unit(SYMS_Arena *arena, SYMS_UnitAccel *unit){
+  SYMS_ProfBegin("syms_tls_var_sid_array_from_unit");
+  SYMS_SymbolIDArray result = {0};
+  switch (unit->format){
+    case SYMS_FileFormat_PDB:
+    {
+      result = syms_pdb_tls_var_sid_array_from_unit(arena, &unit->pdb_accel);
+    }break;
+    
+    case SYMS_FileFormat_DWARF:
+    {
+      // TODO(nick): TLS support on dwarf
+    }break;
+  }
+  SYMS_ASSERT_PARANOID(syms_parser_api_invariants());
+  SYMS_ProfEnd();
+  return(result);
+}
+
+
 #endif //SYMS_PARSER_C
diff --git a/Engine/Source/Runtime/SymsLib/syms/core/syms_parser.h b/Engine/Source/Runtime/SymsLib/syms/core/syms_parser.h
--- a/Engine/Source/Runtime/SymsLib/syms/core/syms_parser.h
+++ b/Engine/Source/Runtime/SymsLib/syms/core/syms_parser.h
@@ -184,6 +184,13 @@ SYMS_API SYMS_USID         syms_type_from_var_sid(SYMS_String8 data, SYMS_DbgAcc
 SYMS_API SYMS_U64          syms_voff_from_var_sid(SYMS_String8 data, SYMS_DbgAccel *dbg, SYMS_UnitAccel *unit,
                                                   SYMS_SymbolID sid);
 
+SYMS_API SYMS_Location     syms_location_from_var_sid(SYMS_Arena *arena, SYMS_String8 data, SYMS_DbgAccel *dbg,
+                                                      SYMS_UnitAccel *unit, SYMS_SymbolID sid);
+SYMS_API SYMS_LocRangeArray syms_location_ranges_from_var_sid(SYMS_Arena *arena, SYMS_String8 data, SYMS_DbgAccel *dbg,
+                                                              SYMS_UnitAccel *unit, SYMS_SymbolID sid);
+SYMS_API SYMS_Location     syms_location_from_id(SYMS_Arena *arena, SYMS_String8 data, SYMS_DbgAccel *dbg,
+                                                 SYMS_UnitAccel *unit, SYMS_LocID loc_id);
+
 // member info
 SYMS_API SYMS_MemsAccel*   syms_mems_accel_from_sid(SYMS_Arena *arena, SYMS_String8 data, SYMS_DbgAccel *dbg,
                                                     SYMS_UnitAccel *unit, SYMS_SymbolID sid);
@@ -211,12 +218,21 @@ SYMS_API SYMS_UnitIDAndSig syms_sig_handle_from_proc_sid(SYMS_String8 data, SYMS
 SYMS_API SYMS_SigInfo      syms_sig_info_from_handle(SYMS_Arena *arena, SYMS_String8 data, SYMS_DbgAccel *dbg,
                                                      SYMS_UnitAccel *unit, SYMS_SigHandle handle);
 
-SYMS_API SYMS_U64RangeArray syms_proc_vranges_from_sid(SYMS_Arena *arena, SYMS_String8 data, SYMS_DbgAccel *dbg,
-                                                       SYMS_UnitAccel *unit, SYMS_SymbolID sid);
+#define syms_proc_vranges_from_sid syms_scope_vranges_from_sid
+SYMS_API SYMS_U64RangeArray syms_scope_vranges_from_sid(SYMS_Arena *arena, SYMS_String8 data, SYMS_DbgAccel *dbg,
+                                                        SYMS_UnitAccel *unit, SYMS_SymbolID sid);
 SYMS_API SYMS_SymbolIDArray syms_scope_children_from_sid(SYMS_Arena *arena, SYMS_String8 data,
                                                          SYMS_DbgAccel *dbg, SYMS_UnitAccel *unit,
                                                          SYMS_SymbolID sid);
 
+SYMS_API SYMS_Location     syms_location_from_proc_sid(SYMS_Arena *arena, SYMS_String8 data, SYMS_DbgAccel *dbg,
+                                                       SYMS_UnitAccel *unit, SYMS_SymbolID sid,
+                                                       SYMS_ProcLoc proc_loc);
+SYMS_API SYMS_LocRangeArray syms_location_ranges_from_proc_sid(SYMS_Arena *arena, SYMS_String8 data,
+                                                               SYMS_DbgAccel *dbg,
+                                                               SYMS_UnitAccel *unit, SYMS_SymbolID sid,
+                                                               SYMS_ProcLoc proc_loc);
+
 SYMS_API SYMS_StrippedInfoArray syms_stripped_from_unit(SYMS_Arena *arena, SYMS_String8 data, SYMS_DbgAccel *dbg,
                                                         SYMS_UnitAccel *unit);
 
@@ -241,6 +257,11 @@ SYMS_API SYMS_UnitID       syms_partner_uid_from_map(SYMS_MapAccel *map);
 SYMS_API SYMS_USIDList     syms_symbol_list_from_string(SYMS_Arena *arena, SYMS_String8 data, SYMS_DbgAccel *dbg,
                                                         SYMS_MapAndUnit map_and_unit, SYMS_String8 string);
 
+// thread vars
+
+SYMS_API SYMS_UnitID        syms_tls_var_uid_from_dbg(SYMS_DbgAccel *dbg);
+SYMS_API SYMS_SymbolIDArray syms_tls_var_sid_array_from_unit(SYMS_Arena *arena, SYMS_UnitAccel *thread_unit);
+
 SYMS_C_LINKAGE_END
 
 #endif //SYMS_PARSER_H
diff --git a/Engine/Source/Runtime/SymsLib/syms/core/unwind/syms_unwind_elf_x64.c b/Engine/Source/Runtime/SymsLib/syms/core/unwind/syms_unwind_elf_x64.c
--- a/Engine/Source/Runtime/SymsLib/syms/core/unwind/syms_unwind_elf_x64.c
+++ b/Engine/Source/Runtime/SymsLib/syms/core/unwind/syms_unwind_elf_x64.c
@@ -1,4 +1,5 @@
 // Copyright Epic Games, Inc. All Rights Reserved.
+
 #ifndef SYMS_UNWIND_ELF_X64_C
 #define SYMS_UNWIND_ELF_X64_C
 
@@ -29,7 +30,7 @@ syms_unwind_elf_x64(SYMS_String8 bin_data, SYMS_ElfBinAccel *bin, SYMS_U64 bin_b
   SYMS_U64 ip_voff = ip_value - rebase_voff_to_vaddr;
   
   //- get sections
-  SYMS_ElfSection *sec_text = syms_elf_sec_from_bin_name__unstable(bin, syms_str8_lit(".text"));
+  SYMS_ElfSection *sec_text = syms_elf_sec_from_bin_name__unstable(bin, syms_str8_lit(".text")); // TODO(nick): What if ELF has two sections with instructions and pointer is ecnoded relative to .text2?
   SYMS_ElfSection *sec_frame_info = syms_elf_sec_from_bin_name__unstable(bin, syms_str8_lit(".eh_frame"));
   SYMS_ElfSection *sec_extra_data = syms_elf_sec_from_bin_name__unstable(bin, syms_str8_lit(".eh_frame_hdr"));
   
@@ -68,8 +69,13 @@ syms_unwind_elf_x64(SYMS_String8 bin_data, SYMS_ElfBinAccel *bin, SYMS_U64 bin_b
     ptr_ctx.raw_base_vaddr = frame_base_voff;
     ptr_ctx.text_vaddr = text_base_vaddr;
     ptr_ctx.data_vaddr = data_base_vaddr;
-    cfi_recs = syms_unwind_elf_x64__eh_frame_cfi_from_ip__sloppy(frame_base, frame_range,
-                                                                 &ptr_ctx, ip_voff);
+    ptr_ctx.func_vaddr = 0;
+    if (sec_extra_data != 0){
+      cfi_recs = syms_unwind_elf_x64__eh_frame_hdr_from_ip(bin_data.str, sec_extra_data->file_range, sec_frame_info->file_range, &ptr_ctx, ip_voff);
+    }
+    else{
+      cfi_recs = syms_unwind_elf_x64__eh_frame_cfi_from_ip__sloppy(frame_base, frame_range, &ptr_ctx, ip_voff);
+    }
   }
   
   //- check cfi records
@@ -357,7 +363,7 @@ syms_unwind_elf_x64__parse_pointer(void *frame_base, SYMS_U64Range frame_range,
   // aligned offset
   SYMS_U64 pointer_off = off;
   if (encoding == SYMS_DwEhPtrEnc_ALIGNED){
-    pointer_off = (off + 7) & ~7;
+    pointer_off = (off + 7) & ~7; // TODO(nick): align to 4 bytes when we parse x86 ELF binary
     encoding = SYMS_DwEhPtrEnc_PTR;
   }
   
@@ -378,7 +384,11 @@ syms_unwind_elf_x64__parse_pointer(void *frame_base, SYMS_U64Range frame_range,
       after_pointer_off = pointer_off + size_param;
     }break;
     
-    case SYMS_DwEhPtrEnc_SIGNED:size_param = 8; goto sfixed;
+    // TODO(nick): SIGNED is actually just a flag that indicates this int is negavite.
+    // There shouldn't be a read when for SIGNED.
+    // For instance, (SYMS_Dw_EhPtrEnc_UDATA2 | SYMS_DwEhPtrEnc_SIGNED) == SYMS_DwEhPtrEnc_SDATA etc.
+    case SYMS_DwEhPtrEnc_SIGNED:size_param = 8; goto sfixed; 
+    
     case SYMS_DwEhPtrEnc_SDATA2:size_param = 2; goto sfixed;
     case SYMS_DwEhPtrEnc_SDATA4:size_param = 4; goto sfixed;
     case SYMS_DwEhPtrEnc_SDATA8:size_param = 8; goto sfixed;
@@ -482,6 +492,12 @@ syms_unwind_elf_x64__eh_frame_parse_cie(void *base, SYMS_U64Range range, SYMS_Dw
       after_ret_addr_reg_off = ret_addr_reg_off + ret_addr_reg_size;
     }
     
+    // TODO(nick): 
+    // Handle "eh" param, it indicates presence of EH Data field.
+    // On 32bit arch it is a 4-byte and on 64-bit 8-byte value.
+    // Reference: https://refspecs.linuxfoundation.org/LSB_3.0.0/LSB-PDA/LSB-PDA/ehframechpt.html
+    // Reference doc doesn't clarify structure for EH Data though
+    
     // check for augmentation data
     SYMS_U64 aug_size_off = after_ret_addr_reg_off;
     SYMS_U64 after_aug_size_off = after_ret_addr_reg_off;
@@ -720,6 +736,119 @@ syms_unwind_elf_x64__eh_frame_cfi_from_ip__sloppy(void *base, SYMS_U64Range rang
   return(result);
 }
 
+SYMS_API SYMS_U64
+syms_search_eh_frame_hdr__linear(void *base, SYMS_U64Range range, SYMS_DwEhPtrCtx *ptr_ctx, SYMS_U64 location)
+{
+  // Table contains only addresses for first instruction in a function and we cannot
+  // guarantee that result is FDE that corresponds to the input location. 
+  // So input location must be cheked against range from FDE header again.
+  
+  SYMS_U64 closest_location = SYMS_U64_MAX;
+  SYMS_U64 closest_address  = SYMS_U64_MAX;
+  
+  SYMS_U64 read_offset = 0;
+  
+  SYMS_U8 version = 0;
+  read_offset += syms_based_range_read_struct(base, range, read_offset, &version);
+  
+  if (version == 1){
+#if 0
+    SYMS_DwEhPtrCtx ptr_ctx; syms_memzero_struct(&ptr_ctx);
+    // Set this to base address of .eh_frame_hdr. Entries are relative
+    // to this section for some reason.
+    ptr_ctx.data_vaddr = range.min;
+    // If input location is VMA then set this to address of .text. 
+    // Pointer parsing function will adjust "init_location" to correct VMA.
+    ptr_ctx.text_vaddr = 0; 
+#endif
+    
+    SYMS_DwEhPtrEnc eh_frame_ptr_enc = 0, fde_count_enc = 0, table_enc = 0;
+    read_offset += syms_based_range_read_struct(base, range, read_offset, &eh_frame_ptr_enc);
+    read_offset += syms_based_range_read_struct(base, range, read_offset, &fde_count_enc);
+    read_offset += syms_based_range_read_struct(base, range, read_offset, &table_enc);
+    
+    SYMS_U64 eh_frame_ptr = 0, fde_count = 0;
+    read_offset += syms_unwind_elf_x64__parse_pointer(base, range, ptr_ctx, eh_frame_ptr_enc, read_offset, &eh_frame_ptr);
+    read_offset += syms_unwind_elf_x64__parse_pointer(base, range, ptr_ctx, fde_count_enc, read_offset, &fde_count);
+    
+    for (SYMS_U64 fde_idx = 0; fde_idx < fde_count; ++fde_idx){
+      SYMS_U64 init_location = 0, address = 0;
+      read_offset += syms_unwind_elf_x64__parse_pointer(base, range, ptr_ctx, table_enc, read_offset, &init_location);
+      read_offset += syms_unwind_elf_x64__parse_pointer(base, range, ptr_ctx, table_enc, read_offset, &address);
+      
+      SYMS_S64 current_delta = (SYMS_S64)(location - init_location);
+      SYMS_S64 closest_delta = (SYMS_S64)(location - closest_location);
+      if (0 <= current_delta && current_delta < closest_delta){
+        closest_location = init_location;
+        closest_address  = address;
+      }
+    }
+  }
+  
+  // address where to find corresponding FDE, this is an absolute offset
+  // into the image file.
+  return closest_address;
+}
+
+SYMS_API SYMS_DwCFIRecords
+syms_unwind_elf_x64__eh_frame_hdr_from_ip(void *base, SYMS_U64Range eh_frame_hdr_range, SYMS_U64Range eh_frame_range, SYMS_DwEhPtrCtx *ptr_ctx, SYMS_U64 ip_voff)
+{
+  SYMS_DwCFIRecords result; syms_memzero_struct(&result);
+  
+  // find FDE offset
+  void *eh_frame_hdr = (void*)((SYMS_U8*)base+eh_frame_hdr_range.min);
+  SYMS_U64 fde_offset = syms_search_eh_frame_hdr__linear(eh_frame_hdr, syms_make_u64_range(0, syms_u64_range_size(eh_frame_hdr_range)), ptr_ctx, ip_voff);
+  
+  SYMS_B32 is_fde_offset_valid = (fde_offset != SYMS_U64_MAX);
+  if (is_fde_offset_valid){
+    SYMS_U64 fde_read_offset = (fde_offset - ptr_ctx->raw_base_vaddr);
+    
+    // read FDE size
+    SYMS_U64 fde_size = 0;
+    fde_read_offset += syms_dw_based_range_read_length(base, eh_frame_range, fde_read_offset, &fde_size);
+    
+    // read FDE discriminator
+    SYMS_U32 fde_discrim = 0;
+    fde_read_offset += syms_based_range_read_struct(base, eh_frame_range, fde_read_offset, &fde_discrim);
+    
+    // compute parent CIE offset
+    SYMS_U64 cie_read_offset = fde_read_offset - (fde_discrim + sizeof(fde_discrim));
+    
+    // read CIE size
+    SYMS_U64 cie_size = 0;
+    cie_read_offset += syms_dw_based_range_read_length(base, eh_frame_range, cie_read_offset, &cie_size);
+    
+    // read CIE discriminator
+    SYMS_U32 cie_discrim = SYMS_U32_MAX;
+    cie_read_offset += syms_based_range_read_struct(base, eh_frame_range, cie_read_offset, &cie_discrim);
+    
+    SYMS_B32 is_fde = (fde_discrim != 0);
+    SYMS_B32 is_cie = (cie_discrim == 0);
+    if (is_fde && is_cie) {
+      void *eh_frame = (void*)((SYMS_U8*)base + eh_frame_range.min);
+      SYMS_U64Range cie_range = syms_make_u64_range(0, cie_read_offset + (cie_size - sizeof(cie_discrim)));
+      SYMS_U64Range fde_range = syms_make_u64_range(0, fde_read_offset + (fde_size - sizeof(fde_discrim)));
+      
+      // parse CIE
+      SYMS_DwCIEUnpacked cie; syms_memzero_struct(&cie);
+      syms_unwind_elf_x64__eh_frame_parse_cie(eh_frame, cie_range, ptr_ctx, cie_read_offset, &cie);
+      
+      // parse FDE
+      SYMS_DwFDEUnpacked fde; syms_memzero_struct(&fde);
+      syms_unwind_elf_x64__eh_frame_parse_fde(eh_frame, fde_range, ptr_ctx, &cie, fde_read_offset, &fde);
+      
+      // range check instruction pointer
+      if (fde.ip_voff_range.min <= ip_voff && ip_voff < fde.ip_voff_range.max){
+        result.valid = syms_true;
+        result.cie = cie;
+        result.fde = fde;
+      }
+    }
+  }
+  
+  return(result);
+}
+
 //- cfi machine
 
 SYMS_API SYMS_DwCFIMachine
diff --git a/Engine/Source/Runtime/SymsLib/syms/core/unwind/syms_unwind_elf_x64.h b/Engine/Source/Runtime/SymsLib/syms/core/unwind/syms_unwind_elf_x64.h
--- a/Engine/Source/Runtime/SymsLib/syms/core/unwind/syms_unwind_elf_x64.h
+++ b/Engine/Source/Runtime/SymsLib/syms/core/unwind/syms_unwind_elf_x64.h
@@ -29,23 +29,24 @@ enum{
   SYMS_DwEhPtrEnc_SDATA8  = 0x0C, // Signed 64-bit value
 };
 enum{
-  SYMS_DwEhPtrEnc_MODIF_MASK = 0xF0,
+  SYMS_DwEhPtrEnc_MODIF_MASK = 0x70,
   SYMS_DwEhPtrEnc_PCREL   = 0x10, // Value is relative to the current program counter.
   SYMS_DwEhPtrEnc_TEXTREL = 0x20, // Value is relative to the .text section.
   SYMS_DwEhPtrEnc_DATAREL = 0x30, // Value is relative to the .got or .eh_frame_hdr section.
   SYMS_DwEhPtrEnc_FUNCREL = 0x40, // Value is relative to the function.
   SYMS_DwEhPtrEnc_ALIGNED = 0x50, // Value is aligned to an address unit sized boundary.
 };
-// TODO(allen): SYMS_DwEhPtrEnc_INDIRECT = 0x80,
 enum{
-  SYMS_DwEhPtrEnc_OMIT    = 0xFF,
+  SYMS_DwEhPtrEnc_INDIRECT = 0x80, // This flag indicates that value is stored in virtual memory.
+
+  SYMS_DwEhPtrEnc_OMIT     = 0xFF,
 };
 
 typedef struct SYMS_DwEhPtrCtx{
-  SYMS_U64 raw_base_vaddr;
-  SYMS_U64 text_vaddr;
-  SYMS_U64 data_vaddr;
-  SYMS_U64 func_vaddr;
+  SYMS_U64 raw_base_vaddr; // address where pointer is being read
+  SYMS_U64 text_vaddr;     // base address of section with instructions (used for encoding pointer on SH and IA64)
+  SYMS_U64 data_vaddr;     // base address of data section (used for encoding pointer on x86-64)
+  SYMS_U64 func_vaddr;     // base address of function where IP is located
 } SYMS_DwEhPtrCtx;
 
 // CIE: Common Information Entry
@@ -191,6 +192,12 @@ SYMS_API void syms_unwind_elf_x64__eh_frame_parse_fde(void *base,SYMS_U64Range r
 SYMS_API SYMS_DwCFIRecords
 syms_unwind_elf_x64__eh_frame_cfi_from_ip__sloppy(void *base, SYMS_U64Range range, SYMS_DwEhPtrCtx *ptr_ctx,
                                                   SYMS_U64 ip_voff);
+SYMS_API SYMS_DwCFIRecords
+syms_unwind_elf_x64__eh_frame_hdr_from_ip(void *base, 
+                                          SYMS_U64Range eh_frame_hdr_range, 
+                                          SYMS_U64Range eh_frame_range, 
+                                          SYMS_DwEhPtrCtx *ptr_ctx, 
+                                          SYMS_U64 ip_voff);
 
 //- cfi machine
 
