- Updated Sysmlib that can handle > 4gb pdb files

diff --git a/Engine/Source/Runtime/SymsLib/syms/overrides/syms_crt_overrides.h b/Engine/Source/Runtime/SymsLib/syms/overrides/syms_crt_overrides.h
--- a/Engine/Source/Runtime/SymsLib/syms/overrides/syms_crt_overrides.h
+++ b/Engine/Source/Runtime/SymsLib/syms/overrides/syms_crt_overrides.h
@@ -1,9 +1,14 @@
 // Copyright Epic Games, Inc. All Rights Reserved.
-/* date = April 6th 2021 4:25 pm */
 
 #ifndef SYMS_CRT_DEFAULT_OVERRIDES_H
 #define SYMS_CRT_DEFAULT_OVERRIDES_H
 
+#if !defined(SYMS_OVERRIDE_FUNC)
+# define SYMS_OVERRIDE_FUNC
+#endif
+
+// NOTE(allen): basic types
+
 #ifdef _MSC_VER
 
 typedef __int8 SYMS_S8;
diff --git a/Engine/Source/Runtime/SymsLib/syms/overrides/syms_default_arena.c b/Engine/Source/Runtime/SymsLib/syms/overrides/syms_default_arena.c
--- a/Engine/Source/Runtime/SymsLib/syms/overrides/syms_default_arena.c
+++ b/Engine/Source/Runtime/SymsLib/syms/overrides/syms_default_arena.c
@@ -1,8 +1,9 @@
 // Copyright Epic Games, Inc. All Rights Reserved.
+
 #ifndef SYMS_DEFAULT_ARENA_C
 #define SYMS_DEFAULT_ARENA_C
 
-static SYMS_DefArena*
+SYMS_OVERRIDE_FUNC SYMS_DefArena*
 syms_arena_def_alloc__sized(SYMS_U64 res, SYMS_U64 cmt){
   SYMS_ASSERT(SYMS_ARENA_HEADER_SIZE < cmt && cmt <= res);
   SYMS_Arena *result = 0;
@@ -26,13 +27,13 @@ syms_arena_def_alloc__sized(SYMS_U64 res, SYMS_U64 cmt){
   return(result);
 }
 
-SYMS_API SYMS_DefArena*
+SYMS_OVERRIDE_FUNC SYMS_DefArena*
 syms_arena_def_alloc(void){
   SYMS_DefArena *result = syms_arena_def_alloc__sized(SYMS_ARENA_RESERVE_SIZE, SYMS_ARENA_COMMIT_SIZE);
   return(result);
 }
 
-SYMS_API void
+SYMS_OVERRIDE_FUNC void
 syms_arena_def_release(SYMS_DefArena *arena){
   for (SYMS_DefArena *node = arena->current, *prev = 0;
        node != 0;
@@ -44,7 +45,7 @@ syms_arena_def_release(SYMS_DefArena *arena){
 
 #define SYMS_ARENA_VERY_LARGE ((SYMS_ARENA_RESERVE_SIZE - SYMS_ARENA_HEADER_SIZE)/2) + 1
 
-SYMS_API void*
+SYMS_OVERRIDE_FUNC void*
 syms_arena_def_push(SYMS_DefArena *arena, SYMS_U64 size){
   void *result = 0;
   
@@ -115,14 +116,14 @@ syms_arena_def_push(SYMS_DefArena *arena, SYMS_U64 size){
   return(result);
 }
 
-SYMS_API SYMS_U64
+SYMS_OVERRIDE_FUNC SYMS_U64
 syms_arena_def_pos(SYMS_Arena *arena){
   SYMS_Arena *current = arena->current;
   SYMS_U64 result = current->base_pos + current->pos;
   return(result);
 }
 
-SYMS_API void
+SYMS_OVERRIDE_FUNC void
 syms_arena_def_pop_to(SYMS_DefArena *arena, SYMS_U64 pos_unclamped){
   SYMS_U64 big_pos = SYMS_ClampBot(SYMS_ARENA_HEADER_SIZE, pos_unclamped);
   
@@ -153,12 +154,12 @@ syms_arena_def_pop_to(SYMS_DefArena *arena, SYMS_U64 pos_unclamped){
 #endif
 }
 
-SYMS_API void
+SYMS_OVERRIDE_FUNC void
 syms_arena_def_set_auto_align(SYMS_DefArena *arena, SYMS_U64 pow2_align){
   arena->align = pow2_align;
 }
 
-SYMS_API void
+SYMS_OVERRIDE_FUNC void
 syms_arena_def_absorb(SYMS_DefArena *arena, SYMS_DefArena *sub){
   // base adjustment
   SYMS_DefArena *current = arena->current;
diff --git a/Engine/Source/Runtime/SymsLib/syms/overrides/syms_default_arena.h b/Engine/Source/Runtime/SymsLib/syms/overrides/syms_default_arena.h
--- a/Engine/Source/Runtime/SymsLib/syms/overrides/syms_default_arena.h
+++ b/Engine/Source/Runtime/SymsLib/syms/overrides/syms_default_arena.h
@@ -1,5 +1,4 @@
 // Copyright Epic Games, Inc. All Rights Reserved.
-/* date = July 28th 2021 3:11 pm */
 
 #ifndef SYMS_DEFAULT_ARENA_H
 #define SYMS_DEFAULT_ARENA_H
@@ -51,13 +50,13 @@ typedef struct SYMS_DefArena{
   struct SYMS_ArenaDev *dev;
 } SYMS_DefArena;
 
-static  SYMS_DefArena*   syms_arena_def_alloc__sized(SYMS_U64 init_res, SYMS_U64 init_cmt);
-SYMS_API SYMS_DefArena*   syms_arena_def_alloc(void);
-SYMS_API void             syms_arena_def_release(SYMS_DefArena *arena);
-SYMS_API void*            syms_arena_def_push(SYMS_DefArena *arena, SYMS_U64 size);
-SYMS_API void             syms_arena_def_pop_to(SYMS_DefArena *arena, SYMS_U64 pos);
-SYMS_API void             syms_arena_def_set_auto_align(SYMS_DefArena *arena, SYMS_U64 pow2_align);
-SYMS_API void             syms_arena_def_absorb(SYMS_DefArena *arena, SYMS_DefArena *sub);
+SYMS_OVERRIDE_FUNC SYMS_DefArena*   syms_arena_def_alloc__sized(SYMS_U64 init_res, SYMS_U64 init_cmt);
+SYMS_OVERRIDE_FUNC SYMS_DefArena*   syms_arena_def_alloc(void);
+SYMS_OVERRIDE_FUNC void             syms_arena_def_release(SYMS_DefArena *arena);
+SYMS_OVERRIDE_FUNC void*            syms_arena_def_push(SYMS_DefArena *arena, SYMS_U64 size);
+SYMS_OVERRIDE_FUNC void             syms_arena_def_pop_to(SYMS_DefArena *arena, SYMS_U64 pos);
+SYMS_OVERRIDE_FUNC void             syms_arena_def_set_auto_align(SYMS_DefArena *arena, SYMS_U64 pow2_align);
+SYMS_OVERRIDE_FUNC void             syms_arena_def_absorb(SYMS_DefArena *arena, SYMS_DefArena *sub);
 
 #define SYMS_Arena SYMS_DefArena
 #define syms_arena_alloc__impl          syms_arena_def_alloc
diff --git a/Engine/Source/Runtime/SymsLib/syms/overrides/syms_default_scratch.c b/Engine/Source/Runtime/SymsLib/syms/overrides/syms_default_scratch.c
--- a/Engine/Source/Runtime/SymsLib/syms/overrides/syms_default_scratch.c
+++ b/Engine/Source/Runtime/SymsLib/syms/overrides/syms_default_scratch.c
@@ -1,4 +1,5 @@
 // Copyright Epic Games, Inc. All Rights Reserved.
+
 #ifndef SYMS_DEFAULT_SCRATCH_C
 #define SYMS_DEFAULT_SCRATCH_C
 
@@ -7,7 +8,7 @@
 
 SYMS_THREAD_LOCAL SYMS_DefaultScratchPool syms_scratch_pool = {0};
 
-static SYMS_Arena*
+SYMS_OVERRIDE_FUNC SYMS_Arena*
 syms_default_get_implicit_thread_arena(SYMS_Arena **conflicts, SYMS_U64 conflict_count){
   // init pool if first time
   if (syms_scratch_pool.arenas[0] == 0){
diff --git a/Engine/Source/Runtime/SymsLib/syms/overrides/syms_default_scratch.h b/Engine/Source/Runtime/SymsLib/syms/overrides/syms_default_scratch.h
--- a/Engine/Source/Runtime/SymsLib/syms/overrides/syms_default_scratch.h
+++ b/Engine/Source/Runtime/SymsLib/syms/overrides/syms_default_scratch.h
@@ -1,5 +1,4 @@
 // Copyright Epic Games, Inc. All Rights Reserved.
-/* date = July 28th 2021 4:28 pm */
 
 #ifndef SYMS_DEFAULT_SCRATCH_H
 #define SYMS_DEFAULT_SCRATCH_H
@@ -14,7 +13,7 @@ typedef struct SYMS_DefaultScratchPool{
 ////////////////////////////////
 // NOTE(allen): Scratch Override
 
-static SYMS_Arena* syms_default_get_implicit_thread_arena(SYMS_Arena **conflicts, SYMS_U64 conflict_count);
+SYMS_OVERRIDE_FUNC SYMS_Arena* syms_default_get_implicit_thread_arena(SYMS_Arena **conflicts, SYMS_U64 conflict_count);
 #define syms_get_implicit_thread_arena__impl syms_default_get_implicit_thread_arena
 #define syms_scratch_pool_tidy__impl()       ((void)0)
 
diff --git a/Engine/Source/Runtime/SymsLib/syms/syms_inc.c b/Engine/Source/Runtime/SymsLib/syms/syms_inc.c
--- a/Engine/Source/Runtime/SymsLib/syms/syms_inc.c
+++ b/Engine/Source/Runtime/SymsLib/syms/syms_inc.c
@@ -1,4 +1,5 @@
 // Copyright Epic Games, Inc. All Rights Reserved.
+
 #ifndef SYMS_INC_C
 #define SYMS_INC_C
 
@@ -9,6 +10,9 @@
 #include "syms/core/base/syms_base.c"
 #include "syms/core/syms_debug_info.c"
 
+// eval
+#include "syms/core/syms_eval.c"
+
 // "windows related" formats
 #include "syms/core/pe/syms_pe.c"
 #include "syms/core/pdb/syms_pdb.c"
@@ -31,16 +35,19 @@
 #include "syms/core/elf/syms_elf_parser.c"
 #include "syms/core/dwarf/syms_dwarf_expr.c"
 #include "syms/core/dwarf/syms_dwarf_parser.c"
+#include "syms/core/dwarf/syms_dwarf_transpiler.c"
 
 // parser abstraction
 #include "syms/core/syms_parser.c"
 #include "syms/core/syms_parser_invariants.c"
 #include "syms/core/data_structures/syms_data_structures.c"
+#include "syms/core/group/syms_type_graph.c"
 #include "syms/core/group/syms_group.c"
 #include "syms/core/file_inf/syms_file_inf.c"
 
 // regs
 #include "syms/core/regs/syms_regs.c"
+#include "syms/core/regs/syms_regs_helpers.c"
 
 // depends on "syms_dwarf_expr" and "regs"
 #include "syms/core/regs/syms_dwarf_regs_helper.c"
@@ -52,4 +59,7 @@
 // serialized type information
 #include "syms/syms_serial_inc.c"
 
+// extras
+#include "syms/core/extras/syms_aggregate_proc_map.c"
+
 #endif // SYMS_INC_C
diff --git a/Engine/Source/Runtime/SymsLib/syms/syms_inc.h b/Engine/Source/Runtime/SymsLib/syms/syms_inc.h
--- a/Engine/Source/Runtime/SymsLib/syms/syms_inc.h
+++ b/Engine/Source/Runtime/SymsLib/syms/syms_inc.h
@@ -11,6 +11,17 @@
 #include "syms/core/base/syms_base.h"
 #include "syms/core/syms_debug_info.h"
 
+// regs
+#include "syms/core/regs/syms_regs.h"
+#pragma pack(push, 1)
+#include "syms/core/generated/syms_meta_regs_x86.h"
+#include "syms/core/generated/syms_meta_regs_x64.h"
+#pragma pack(pop)
+#include "syms/core/regs/syms_regs_helpers.h"
+
+// eval
+#include "syms/core/syms_eval.h"
+
 // "windows related" formats
 #include "syms/core/pe/syms_pe.h"
 #include "syms/core/pdb/syms_pdb.h"
@@ -33,23 +44,16 @@
 #include "syms/core/elf/syms_elf_parser.h"
 #include "syms/core/dwarf/syms_dwarf_expr.h"
 #include "syms/core/dwarf/syms_dwarf_parser.h"
+#include "syms/core/dwarf/syms_dwarf_transpiler.h"
 
 // parser abstraction
 #include "syms/core/syms_parser.h"
 #include "syms/core/syms_parser_invariants.h"
 #include "syms/core/data_structures/syms_data_structures.h"
+#include "syms/core/group/syms_type_graph.h"
 #include "syms/core/group/syms_group.h"
 #include "syms/core/file_inf/syms_file_inf.h"
 
-// regs
-#include "syms/core/regs/syms_regs.h"
-
-#pragma pack(push, 1)
-#include "syms/core/generated/syms_meta_regs_classes.h"
-#include "syms/core/generated/syms_meta_regs_x86.h"
-#include "syms/core/generated/syms_meta_regs_x64.h"
-#pragma pack(pop)
-
 #include "syms/core/regs/syms_regs_x64.h"
 // depends on "syms_dwarf_expr" and "regs"
 #include "syms/core/regs/syms_dwarf_regs_helper.h"
@@ -61,4 +65,7 @@
 // serialized type information
 #include "syms/syms_serial_inc.h"
 
+// extras
+#include "syms/core/extras/syms_aggregate_proc_map.h"
+
 #endif // SYMS_INC_H
