Add support for Linux for SymLibs

diff --git a/Engine/Source/Runtime/SymsLib/BuildForLinux.sh b/Engine/Source/Runtime/SymsLib/BuildForLinux.sh
new file mode 100755
--- /dev/null
+++ b/Engine/Source/Runtime/SymsLib/BuildForLinux.sh
@@ -0,0 +1,89 @@
+#!/bin/bash
+
+set -eu
+
+ARCH=x86_64-unknown-linux-gnu
+SCRIPT_DIR=$(cd "$(dirname "$BASH_SOURCE")" ; pwd)
+UE_THIRD_PARTY_DIR=$(cd "${SCRIPT_DIR}/../../../.."; pwd)
+SYMS_DIR=$(cd "${SCRIPT_DIR}"; pwd)
+
+CFLAGS="-ffunction-sections -fdata-sections"
+CXXFLAGS="-ffunction-sections -fdata-sections -I$UE_THIRD_PARTY_DIR/Unix/LibCxx/include -I$UE_THIRD_PARTY_DIR/Unix/LibCxx/include/c++/v1"
+LIBS="$UE_THIRD_PARTY_DIR/Unix/LibCxx/lib/Unix/${ARCH}/libc++.a $UE_THIRD_PARTY_DIR/Unix/LibCxx/lib/Unix/${ARCH}/libc++abi.a -lm -lc -lgcc_s -lgcc -lpthread"
+
+# Get num of cores
+export CORES=$(getconf _NPROCESSORS_ONLN)
+echo "Using ${CORES} cores for building"
+
+BuildLibSyms()
+{
+    export ARCH=$1
+    export FLAVOR=$2
+    local EXT=$3
+    local BUILD_DIR=/tmp/Build-Syms-${ARCH}-${FLAVOR}
+
+    echo "Building ${ARCH} ${FLAVOR}"
+    rm -rf ${BUILD_DIR}
+    mkdir -p ${BUILD_DIR}/Build
+
+    pushd ${BUILD_DIR}/Build
+
+    set -x
+    cmake \
+      -DCMAKE_TOOLCHAIN_FILE="/tmp/__cmake_toolchain.cmake" \
+      -DCMAKE_MAKE_PROGRAM=$(which make) \
+      -DCMAKE_BUILD_TYPE=${FLAVOR} \
+      -DCMAKE_POSITION_INDEPENDENT_CODE=ON \
+      ${SYMS_DIR}/
+    set +x
+
+    VERBOSE=1 make -j ${CORES}
+
+    mkdir -p ${SYMS_DIR}/lib/Unix/${ARCH}
+    cp -v ${BUILD_DIR}/Build/libsyms.a ${SYMS_DIR}/lib/Unix/${ARCH}/libsyms${EXT}_fPIC.a
+
+    popd
+}
+
+( cat <<_EOF_
+  ## autogenerated by ${BASH_SOURCE} script
+  SET(LINUX_MULTIARCH_ROOT \$ENV{UE_SDKS_ROOT}/HostLinux/Linux_x64/v19_clang-11.0.1-centos7)
+  SET(ARCHITECTURE_TRIPLE \$ENV{ARCH})
+
+  message (STATUS "LINUX_MULTIARCH_ROOT is '\${LINUX_MULTIARCH_ROOT}'")
+  message (STATUS "ARCHITECTURE_TRIPLE is '\${ARCHITECTURE_TRIPLE}'")
+
+  SET(CMAKE_CROSSCOMPILING TRUE)
+  SET(CMAKE_SYSTEM_NAME Linux)
+  SET(CMAKE_SYSTEM_VERSION 1)
+
+  # sysroot
+  SET(CMAKE_SYSROOT \${LINUX_MULTIARCH_ROOT}/\${ARCHITECTURE_TRIPLE})
+
+  SET(CMAKE_LIBRARY_ARCHITECTURE \${ARCHITECTURE_TRIPLE})
+
+  # specify the cross compiler
+  SET(CMAKE_C_COMPILER            \${CMAKE_SYSROOT}/bin/clang)
+  SET(CMAKE_C_COMPILER_TARGET     \${ARCHITECTURE_TRIPLE})
+  SET(CMAKE_C_FLAGS "-target      \${ARCHITECTURE_TRIPLE} ${CFLAGS}")
+
+  SET(CMAKE_CXX_COMPILER          \${CMAKE_SYSROOT}/bin/clang++)
+  SET(CMAKE_CXX_COMPILER_TARGET   \${ARCHITECTURE_TRIPLE})
+  SET(CMAKE_CXX_FLAGS "-target    \${ARCHITECTURE_TRIPLE} ${CXXFLAGS}")
+
+  SET(CMAKE_SHARED_LINKER_FLAGS   "${LIBS}")
+
+  SET(CMAKE_ASM_COMPILER          \${CMAKE_SYSROOT}/bin/clang)
+
+  SET(CMAKE_FIND_ROOT_PATH        \${LINUX_MULTIARCH_ROOT})
+
+  # hoping to force it to use ar
+  set(CMAKE_FIND_ROOT_PATH_MODE_PROGRAM ONLY)
+  set(CMAKE_FIND_ROOT_PATH_MODE_LIBRARY ONLY)
+  set(CMAKE_FIND_ROOT_PATH_MODE_INCLUDE ONLY)
+
+_EOF_
+) > /tmp/__cmake_toolchain.cmake
+
+BuildLibSyms ${ARCH} Release ""
+BuildLibSyms ${ARCH} Debug "d"
diff --git a/Engine/Source/Runtime/SymsLib/CMakeLists.txt b/Engine/Source/Runtime/SymsLib/CMakeLists.txt
new file mode 100644
--- /dev/null
+++ b/Engine/Source/Runtime/SymsLib/CMakeLists.txt
@@ -0,0 +1,9 @@
+cmake_minimum_required(VERSION 3.00)
+PROJECT(SYMLIBS)
+
+set(SRC_FILES
+    syms/symslib.c
+)
+
+add_library(syms STATIC ${SRC_FILES})
+target_include_directories(syms PUBLIC ${CMAKE_SOURCE_DIR})
diff --git a/Engine/Source/Runtime/SymsLib/SymsLib.Build.cs b/Engine/Source/Runtime/SymsLib/SymsLib.Build.cs
--- a/Engine/Source/Runtime/SymsLib/SymsLib.Build.cs
+++ b/Engine/Source/Runtime/SymsLib/SymsLib.Build.cs
@@ -1,13 +1,14 @@
 // Copyright Epic Games, Inc. All Rights Reserved.
 
 using UnrealBuildTool;
+using System.IO;
 
 public class SymsLib : ModuleRules
 {
 	public SymsLib(ReadOnlyTargetRules Target) : base(Target)
 	{
 		Type = ModuleType.External;
-		
+
 		string LibPathBase =  ModuleDirectory + "/lib";
 
 		PublicSystemIncludePaths.Add(ModuleDirectory);
@@ -19,5 +20,13 @@ public class SymsLib : ModuleRules
 			string LibPath = LibPathBase + "/x64" + (bUseDebug ? "/Debug" : "/Release");
 			PublicAdditionalLibraries.Add(LibPath + "/SymsLib.lib");
 		}
+		else if (Target.Platform.IsInGroup(UnrealPlatformGroup.Unix))
+		{
+			bool bUseDebug = Target.Configuration == UnrealTargetConfiguration.Debug;
+			string LibName = bUseDebug ? "libsymsd_fPIC.a" : "libsyms_fPIC.a";
+			string LibPath = Path.Combine(LibPathBase, "Unix", Target.Architecture);
+
+			PublicAdditionalLibraries.Add(Path.Combine(LibPath, LibName));
+		}
 	}
 }
diff --git a/Engine/Source/Runtime/SymsLib/syms/core/group/syms_group.h b/Engine/Source/Runtime/SymsLib/syms/core/group/syms_group.h
--- a/Engine/Source/Runtime/SymsLib/syms/core/group/syms_group.h
+++ b/Engine/Source/Runtime/SymsLib/syms/core/group/syms_group.h
@@ -143,7 +143,7 @@ typedef struct SYMS_Group{
 
 SYMS_READ_ONLY SYMS_GLOBAL SYMS_SymbolIDArray syms_sid_array_nil = {0};
 SYMS_READ_ONLY SYMS_GLOBAL SYMS_String8Array syms_string_array_nil = {0};
-SYMS_READ_ONLY SYMS_GLOBAL SYMS_LineParseOut syms_line_parse_nil = {0};
+SYMS_READ_ONLY SYMS_GLOBAL SYMS_LineParseOut syms_line_parse_nil = {{0}};
 SYMS_READ_ONLY SYMS_GLOBAL SYMS_SpatialMap1D syms_spatial_map_1d_nil = {0};
 SYMS_READ_ONLY SYMS_GLOBAL SYMS_LineToAddrMap syms_line_to_addr_map_nil = {0};
 SYMS_READ_ONLY SYMS_GLOBAL SYMS_TypeMemberArray syms_type_member_array_nil = {0};
