- Updated symslib for getting symbol info by name instead of location

diff --git a/Engine/Source/Developer/TraceServices/Private/Model/SymslibResolver.cpp b/Engine/Source/Developer/TraceServices/Private/Model/SymslibResolver.cpp
--- a/Engine/Source/Developer/TraceServices/Private/Model/SymslibResolver.cpp
+++ b/Engine/Source/Developer/TraceServices/Private/Model/SymslibResolver.cpp
@@ -1000,22 +1000,22 @@ EModuleStatus FSymslibResolver::LoadModule(FModuleEntry* Entry, FStringView Sear
 
 	// store stripped format symbols
 	{
-		SYMS_StrippedInfoArray StrippedInfo = syms_group_stripped_info(Group);
+		SYMS_LinkNameRecArray StrippedInfo = syms_group_link_name_records(Group);
 
 		FSymsSymbol* StrippedSymbols = syms_push_array(Arena, FSymsSymbol, StrippedInfo.count);
 		for (SYMS_U64 Index = 0; Index < StrippedInfo.count; Index++)
 		{
-			SYMS_StrippedInfo* Info = &StrippedInfo.info[Index];
+			SYMS_LinkNameRec* Info = &StrippedInfo.recs[Index];
 			FSymsSymbol* StrippedSymbol = &StrippedSymbols[Index];
 
 			SYMS_String8 Name = syms_push_string_copy(Arena, Info->name);
 			StrippedSymbol->Name = reinterpret_cast<char*>(Name.str);
 		}
 
-		Instance->StrippedMap = syms_spatial_map_1d_copy(Arena, syms_group_stripped_info_map(Group));
+		Instance->StrippedMap = syms_spatial_map_1d_copy(Arena, syms_group_link_name_spatial_map(Group));
 		Instance->StrippedSymbols = StrippedSymbols;
 
-		SymbolCount += StrippedInfo.count;
+		SymbolCount += (uint32)StrippedInfo.count;
 	}
 
 	Instance->UnitMap = syms_spatial_map_1d_copy(Arena, syms_group_unit_map(Group));
diff --git a/Engine/Source/Runtime/SymsLib/syms/core/base/syms_base.c b/Engine/Source/Runtime/SymsLib/syms/core/base/syms_base.c
--- a/Engine/Source/Runtime/SymsLib/syms/core/base/syms_base.c
+++ b/Engine/Source/Runtime/SymsLib/syms/core/base/syms_base.c
@@ -36,7 +36,7 @@ syms_u64_range_size(SYMS_U64Range range){
 }
 
 ////////////////////////////////
-//~ NOTE(allen): Hash Functions
+//~ allen: Hash Functions
 
 SYMS_API SYMS_U64
 syms_hash_djb2(SYMS_String8 string){
@@ -150,7 +150,7 @@ SYMS_API SYMS_U64
 syms_enum_index_from_value_identity(SYMS_U64 v){ return(v); }
 
 ////////////////////////////////
-//~ NOTE(allen): String Functions
+//~ allen: String Functions
 
 SYMS_API SYMS_B32
 syms_codepoint_is_whitespace(SYMS_U32 codepoint){
@@ -418,7 +418,7 @@ syms_string_split(SYMS_Arena *arena, SYMS_String8 input, SYMS_U32 delimiter)
 }
 
 ////////////////////////////////
-//~ NOTE(allen): String -> Int
+//~ allen: String -> Int
 
 SYMS_API SYMS_U64
 syms_u64_from_string(SYMS_String8 str, SYMS_U32 radix){
@@ -549,7 +549,7 @@ syms_string_from_u64(SYMS_Arena *arena, SYMS_U64 x){
 }
 
 ////////////////////////////////
-//~ NOTE(rjf): U64 Range Functions
+//~ rjf: U64 Range Functions
 
 SYMS_API void
 syms_u64_range_list_push_node(SYMS_U64RangeNode *node, SYMS_U64RangeList *list, SYMS_U64Range range){
@@ -599,7 +599,7 @@ syms_u64_range_array_from_list(SYMS_Arena *arena, SYMS_U64RangeList *list){
 }
 
 ////////////////////////////////
-//~ NOTE(allen): Arena Functions
+//~ allen: Arena Functions
 
 SYMS_API void
 syms_arena_push_align(SYMS_Arena *arena, SYMS_U64 pow2_boundary){
@@ -645,7 +645,26 @@ syms_get_scratch(SYMS_Arena **conflicts, SYMS_U64 conflict_count){
 }
 
 ////////////////////////////////
-//~ NOTE(allen): Thread Lanes
+//~ allen: Syms Sort Node
+
+SYMS_API SYMS_SortNode*
+syms_sort_node_push(SYMS_Arena *arena, SYMS_SortNode **stack, SYMS_SortNode **free_stack,
+                    SYMS_U64 first, SYMS_U64 opl){
+  SYMS_SortNode *result = *free_stack;
+  if (result != 0){
+    SYMS_StackPop(*free_stack);
+  }
+  else{
+    result = syms_push_array(arena, SYMS_SortNode, 1);
+  }
+  SYMS_StackPush(*stack, result);
+  result->first = first;
+  result->opl = opl;
+  return(result);
+}
+
+////////////////////////////////
+//~ allen: Thread Lanes
 
 SYMS_THREAD_LOCAL SYMS_U64 syms_thread_lane = 0;
 
@@ -660,7 +679,7 @@ syms_get_lane(void){
 }
 
 ////////////////////////////////
-//~ NOTE(rjf): Based Ranges
+//~ rjf: Based Ranges
 
 SYMS_API void*
 syms_based_range_ptr(void *base, SYMS_U64Range range, SYMS_U64 offset){
@@ -756,7 +775,7 @@ syms_based_range_read_string(void *base, SYMS_U64Range range, SYMS_U64 offset){
 }
 
 ////////////////////////////////
-//~ NOTE(allen): Memory Views
+//~ allen: Memory Views
 
 SYMS_API SYMS_MemoryView
 syms_memory_view_make(SYMS_String8 data, SYMS_U64 base){
@@ -787,7 +806,7 @@ syms_unwind_result_missed_read(SYMS_UnwindResult *unwind_result, SYMS_U64 addr){
 }
 
 ////////////////////////////////
-//~ NOTE(nick): Bit manipulations
+//~ nick: Bit manipulations
 
 SYMS_API SYMS_U16
 syms_bswap_u16(SYMS_U16 x)
@@ -838,7 +857,7 @@ syms_bswap_bytes(void *p, SYMS_U64 size)
 #include "syms/core/generated/syms_meta_serial_base.c"
 
 ////////////////////////////////
-//~ NOTE(allen): Dev Features
+//~ allen: Dev Features
 
 #include "syms_dev.c"
 
diff --git a/Engine/Source/Runtime/SymsLib/syms/core/base/syms_base.h b/Engine/Source/Runtime/SymsLib/syms/core/base/syms_base.h
--- a/Engine/Source/Runtime/SymsLib/syms/core/base/syms_base.h
+++ b/Engine/Source/Runtime/SymsLib/syms/core/base/syms_base.h
@@ -5,7 +5,7 @@
 #define SYMS_BASE_H
 
 ////////////////////////////////
-//~ NOTE(allen): Version
+//~ allen: Version
 
 #define SYMS_VERSION_MAJOR    1
 #define SYMS_VERSION_MINOR    0
@@ -13,17 +13,17 @@
 #define SYMS_VERSION_STR      "1.0.0"
 
 ////////////////////////////////
-//~ NOTE(rjf): Context Cracking
+//~ rjf: Context Cracking
 
 #include "syms_base_context_crack.h"
 
 ////////////////////////////////
-//~ NOTE(rjf): Overrideables
+//~ rjf: Overrideables
 
 #include "syms_base_overrides_check.h"
 
 ////////////////////////////////
-//~ NOTE(allen): Linkage Macros
+//~ allen: Linkage Macros
 
 #if !defined(SYMS_API)
 # define SYMS_API static
@@ -67,7 +67,7 @@
 #endif
 
 ////////////////////////////////
-//~ NOTE(allen): Feature Macros
+//~ allen: Feature Macros
 
 #if !defined(SYMS_PARANOID)
 # define SYMS_PARANOID 0
@@ -92,7 +92,7 @@
 #endif
 
 ////////////////////////////////
-//~ NOTE(rjf): Base Types and Constants
+//~ rjf: Base Types and Constants
 
 typedef SYMS_S8  SYMS_B8;
 typedef SYMS_S16 SYMS_B16;
@@ -107,7 +107,7 @@ typedef double   SYMS_F64;
 typedef SYMS_U32 SYMS_RegID;
 
 ////////////////////////////////
-//~ NOTE(allen): Macros
+//~ allen: Macros
 
 #if !defined(syms_write_srcloc__impl)
 # define syms_write_srcloc__impl(f,l) ((void)0)
@@ -191,7 +191,7 @@ typedef SYMS_U32 SYMS_RegID;
 #define SYMS_THIS_SRCLOC __FILE__ ":" SYMS_Stringify(__LINE__)
 
 ////////////////////////////////
-//~ NOTE(allen): ID Macros
+//~ allen: ID Macros
 
 #define SYMS_ID_u32_0(id) (SYMS_U32)(id)
 #define SYMS_ID_u32_1(id) (SYMS_U32)((id) >> 32)
@@ -204,7 +204,7 @@ typedef SYMS_U32 SYMS_RegID;
 #define SYMS_ID_u16_u16_u32(a,b,c) ((SYMS_U64)(a) | ((SYMS_U64)(b) << 16) | ((SYMS_U64)(c) << 32))
 
 ////////////////////////////////
-//~ NOTE(allen): Linked List Macros
+//~ allen: Linked List Macros
 
 #define SYMS_QueuePush_N(f,l,n,next) ( (f)==0?\
 ((f)=(l)=(n),(n)->next=0):\
@@ -227,7 +227,7 @@ typedef SYMS_U32 SYMS_RegID;
 #define SYMS_StackPop(f) SYMS_StackPop_N(f,next)
 
 ////////////////////////////////
-//~ NOTE(allen): Common Basic Types
+//~ allen: Common Basic Types
 
 typedef struct SYMS_U64Array{
   SYMS_U64 *u64;
@@ -307,12 +307,21 @@ typedef struct SYMS_StringJoin{
 } SYMS_StringJoin;
 
 ////////////////////////////////
-//~ NOTE(allen): Generated Types
+//~ allen: Syms Sort Node
+
+typedef struct SYMS_SortNode{
+  struct SYMS_SortNode *next;
+  SYMS_U64 first;
+  SYMS_U64 opl;
+} SYMS_SortNode;
+
+////////////////////////////////
+//~ allen: Generated Types
 
 #include "syms/core/generated/syms_meta_base.h"
 
 ////////////////////////////////
-//~ NOTE(rjf): Serial Information
+//~ rjf: Serial Information
 
 typedef enum SYMS_SerialWidthKind{
   SYMS_SerialWidthKind_Null,
@@ -366,7 +375,7 @@ typedef struct SYMS_SerialType{
 } SYMS_SerialType;
 
 ////////////////////////////////
-//~ NOTE(allen): Syms Arena
+//~ allen: Syms Arena
 
 typedef struct SYMS_ArenaTemp{
   SYMS_Arena *arena;
@@ -374,10 +383,10 @@ typedef struct SYMS_ArenaTemp{
 } SYMS_ArenaTemp;
 
 ////////////////////////////////
-//~ NOTE(allen): Memory Views
+//~ allen: Memory Views
 
 typedef struct SYMS_MemoryView{
-  // NOTE(allen): Upgrade path:
+  // allen: Upgrade path:
   //  1. A list of ranges like this one
   //  2. After building the list put into a binary-searchable format
   //  3. Equip with the ability to request missing memory in-line.
@@ -395,7 +404,7 @@ typedef struct SYMS_UnwindResult{
 
 
 ////////////////////////////////
-//~ NOTE(allen): Generated Serial Info
+//~ allen: Generated Serial Info
 
 #include "syms/core/generated/syms_meta_serial_base.h"
 
@@ -403,19 +412,19 @@ typedef struct SYMS_UnwindResult{
 SYMS_C_LINKAGE_BEGIN
 
 ////////////////////////////////
-//~ NOTE(rjf): Library Metadata
+//~ rjf: Library Metadata
 
 SYMS_API SYMS_String8 syms_version_string(void);
 
 ////////////////////////////////
-//~ NOTE(rjf): Basic Type Functions
+//~ rjf: Basic Type Functions
 
 SYMS_API SYMS_U64Range syms_make_u64_range(SYMS_U64 min, SYMS_U64 max);
 SYMS_API SYMS_U64Range syms_make_u64_inrange(SYMS_U64Range range, SYMS_U64 offset, SYMS_U64 size);
 SYMS_API SYMS_U64 syms_u64_range_size(SYMS_U64Range range);
 
 ////////////////////////////////
-//~ NOTE(allen): Hash Functions
+//~ allen: Hash Functions
 
 SYMS_GLOBAL SYMS_U64 syms_hash_djb2_initial = 5381;
 SYMS_API SYMS_U64 syms_hash_djb2(SYMS_String8 string);
@@ -424,7 +433,7 @@ SYMS_API SYMS_U64 syms_hash_djb2_continue(SYMS_String8 string, SYMS_U64 intermed
 SYMS_API SYMS_U64 syms_hash_u64(SYMS_U64 x);
 
 ////////////////////////////////
-//~ NOTE(rjf): Serial Information Functions
+//~ rjf: Serial Information Functions
 
 #define syms_serial_type(name) (_syms_serial_type_##name)
 #define syms_string_from_enum_value(enum_type, value) \
@@ -443,7 +452,7 @@ SYMS_API SYMS_String8List  syms_string_list_from_flags(SYMS_Arena *arena, SYMS_S
 SYMS_API SYMS_U64 syms_enum_index_from_value_identity(SYMS_U64 v);
 
 ////////////////////////////////
-//~ NOTE(allen): String Functions
+//~ allen: String Functions
 
 SYMS_API SYMS_B32    syms_codepoint_is_whitespace(SYMS_U32 codepoint);
 SYMS_API SYMS_U32    syms_lowercase_from_codepoint(SYMS_U32 codepoint);
@@ -478,7 +487,7 @@ SYMS_API SYMS_String8 syms_string_trunc_symbol_heuristic(SYMS_String8 string);
 SYMS_API SYMS_String8List syms_string_split(SYMS_Arena *arena, SYMS_String8 input, SYMS_U32 delimiter);
 
 ////////////////////////////////
-//~ NOTE(allen): String <-> Integer
+//~ allen: String <-> Integer
 
 SYMS_API SYMS_U64 syms_u64_from_string(SYMS_String8 str, SYMS_U32 radix);
 SYMS_API SYMS_S64 syms_s64_from_string_c_rules(SYMS_String8 str);
@@ -486,7 +495,7 @@ SYMS_API SYMS_S64 syms_s64_from_string_c_rules(SYMS_String8 str);
 SYMS_API SYMS_String8 syms_string_from_u64(SYMS_Arena *arena, SYMS_U64 x);
 
 ////////////////////////////////
-//~ NOTE(rjf): U64 Range Functions
+//~ rjf: U64 Range Functions
 
 SYMS_API void syms_u64_range_list_push_node(SYMS_U64RangeNode *node, SYMS_U64RangeList *list, SYMS_U64Range range);
 SYMS_API void syms_u64_range_list_push(SYMS_Arena *arena, SYMS_U64RangeList *list, SYMS_U64Range range);
@@ -495,7 +504,7 @@ SYMS_API void syms_u64_range_list_concat(SYMS_U64RangeList *list, SYMS_U64RangeL
 SYMS_API SYMS_U64RangeArray syms_u64_range_array_from_list(SYMS_Arena *arena, SYMS_U64RangeList *list);
 
 ////////////////////////////////
-//~ NOTE(rjf): Memory/Arena Functions
+//~ rjf: Memory/Arena Functions
 
 #define syms_arena_alloc          syms_arena_alloc__impl
 #define syms_arena_release        syms_arena_release__impl
@@ -525,13 +534,19 @@ SYMS_API SYMS_ArenaTemp syms_get_scratch(SYMS_Arena **conflicts, SYMS_U64 count)
 #define syms_scratch_pool_tidy syms_scratch_pool_tidy__impl
 
 ////////////////////////////////
-//~ NOTE(allen): Thread Lanes
+//~ allen: Syms Sort Node
+
+SYMS_API SYMS_SortNode* syms_sort_node_push(SYMS_Arena *arena, SYMS_SortNode **stack, SYMS_SortNode **free_stack,
+                                            SYMS_U64 first, SYMS_U64 opl);
+
+////////////////////////////////
+//~ allen: Thread Lanes
 
 SYMS_API void     syms_set_lane(SYMS_U64 lane);
 SYMS_API SYMS_U64 syms_get_lane(void);
 
 ////////////////////////////////
-//~ NOTE(rjf): Based Ranges
+//~ rjf: Based Ranges
 
 SYMS_API void *   syms_based_range_ptr(void *base, SYMS_U64Range range, SYMS_U64 offset);
 SYMS_API SYMS_U64 syms_based_range_read(void *base, SYMS_U64Range range, SYMS_U64 offset, SYMS_U64 out_size, void *out);
@@ -543,7 +558,7 @@ SYMS_API SYMS_String8 syms_based_range_read_string(void *base, SYMS_U64Range ran
 #define syms_based_range_read_struct(b,r,o,p) syms_based_range_read((b), (r), (o), sizeof(*(p)), p)
 
 ////////////////////////////////
-//~ NOTE(allen): Memory Views
+//~ allen: Memory Views
 
 SYMS_API SYMS_MemoryView syms_memory_view_make(SYMS_String8 data, SYMS_U64 base);
 SYMS_API SYMS_B32        syms_memory_view_read(SYMS_MemoryView *memview, SYMS_U64 addr,
@@ -554,7 +569,7 @@ SYMS_API SYMS_B32        syms_memory_view_read(SYMS_MemoryView *memview, SYMS_U6
 SYMS_API void syms_unwind_result_missed_read(SYMS_UnwindResult *unwind_result, SYMS_U64 addr);
 
 ////////////////////////////////
-//~ NOTE(nick): Bit manipulations
+//~ nick: Bit manipulations
 
 SYMS_API SYMS_U16 syms_bswap_u16(SYMS_U16 x);
 SYMS_API SYMS_U32 syms_bswap_u32(SYMS_U32 x);
@@ -562,7 +577,7 @@ SYMS_API SYMS_U64 syms_bswap_u64(SYMS_U64 x);
 SYMS_API void syms_bswap_bytes(void *p, SYMS_U64 size);
 
 ////////////////////////////////
-//~ NOTE(allen): Dev Features
+//~ allen: Dev Features
 
 #include "syms_dev.h"
 
diff --git a/Engine/Source/Runtime/SymsLib/syms/core/data_structures/syms_data_structures.c b/Engine/Source/Runtime/SymsLib/syms/core/data_structures/syms_data_structures.c
--- a/Engine/Source/Runtime/SymsLib/syms/core/data_structures/syms_data_structures.c
+++ b/Engine/Source/Runtime/SymsLib/syms/core/data_structures/syms_data_structures.c
@@ -1313,15 +1313,16 @@ syms_string_array_copy(SYMS_Arena *arena, SYMS_StringCons *cons, SYMS_String8Arr
   return(result);
 }
 
-SYMS_API SYMS_StrippedInfoArray
-syms_stripped_info_copy(SYMS_Arena *arena, SYMS_StrippedInfoArray *stripped){
-  SYMS_ProfBegin("syms_stripped_info_copy");
-  SYMS_StrippedInfoArray result = {0};
-  result.count = stripped->count;
-  result.info = syms_push_array(arena, SYMS_StrippedInfo, result.count);
+
+SYMS_API SYMS_LinkNameRecArray
+syms_link_name_record_copy(SYMS_Arena *arena, SYMS_LinkNameRecArray *array){
+  SYMS_ProfBegin("syms_link_name_record_copy");
+  SYMS_LinkNameRecArray result = {0};
+  result.count = array->count;
+  result.recs = syms_push_array(arena, SYMS_LinkNameRec, result.count);
   
-  SYMS_StrippedInfo *src = stripped->info;
-  SYMS_StrippedInfo *dst = result.info;
+  SYMS_LinkNameRec *src = array->recs;
+  SYMS_LinkNameRec *dst = result.recs;
   for (SYMS_U64 i = 0; i < result.count; i += 1, src += 1, dst += 1){
     dst->name = syms_push_string_copy(arena, src->name);
     dst->voff = src->voff;
diff --git a/Engine/Source/Runtime/SymsLib/syms/core/data_structures/syms_data_structures.h b/Engine/Source/Runtime/SymsLib/syms/core/data_structures/syms_data_structures.h
--- a/Engine/Source/Runtime/SymsLib/syms/core/data_structures/syms_data_structures.h
+++ b/Engine/Source/Runtime/SymsLib/syms/core/data_structures/syms_data_structures.h
@@ -5,7 +5,7 @@
 #define SYMS_DATA_STRUCTURES_H
 
 ////////////////////////////////
-//~ NOTE(allen): String Cons Structure
+//~ allen: String Cons Structure
 
 // deduplicates equivalent strings
 
@@ -40,7 +40,7 @@ typedef struct SYMS_DataIdxCons{
 } SYMS_DataIdxCons;
 
 ////////////////////////////////
-//~ NOTE(allen): U64 Set
+//~ allen: U64 Set
 
 typedef struct SYMS_U64Set{
   SYMS_U64 *vals;
@@ -49,7 +49,7 @@ typedef struct SYMS_U64Set{
 } SYMS_U64Set;
 
 ////////////////////////////////
-//~ NOTE(allen): 1D Spatial Mapping Structure
+//~ allen: 1D Spatial Mapping Structure
 
 // assigns a value to ranges of unsigned 64-bit values
 // ranges specified in half-open intervals: [min,max)
@@ -97,7 +97,7 @@ typedef struct SYMS_1DEndPoint{
 } SYMS_1DEndPoint;
 
 ////////////////////////////////
-//~ NOTE(allen): File Mapping Structure ({UnitID,FileID} -> String)
+//~ allen: File Mapping Structure ({UnitID,FileID} -> String)
 
 // maps a unit-id,file-id pair to a string
 // organized as a hash table to opimize for key based lookups
@@ -118,7 +118,7 @@ typedef struct SYMS_FileID2NameMap{
 } SYMS_FileID2NameMap;
 
 ////////////////////////////////
-//~ NOTE(allen): File Mapping Structure (String -> {UnitID,FileID})
+//~ allen: File Mapping Structure (String -> {UnitID,FileID})
 
 // maps strings to a set of unit-id,file-id pairs
 // oragnized in an array of strings with each string
@@ -168,7 +168,7 @@ typedef struct SYMS_Name2FileIDMapLoose{
 
 
 ////////////////////////////////
-//~ NOTE(allen): ID Mapping Structure
+//~ allen: ID Mapping Structure
 
 // maps unsigned 64-bit values to arbitrary user pointers
 // organized as a hash table to opimize key based lookups
@@ -190,7 +190,7 @@ typedef struct SYMS_IDMap{
 
 
 ////////////////////////////////
-//~ NOTE(allen): Symbol Name Mapping Structure (String -> Array(USID))
+//~ allen: Symbol Name Mapping Structure (String -> Array(USID))
 
 // maps strings to lists of USIDs
 // organized as an array of nodes and a hash table simultaneously
@@ -229,7 +229,7 @@ typedef struct SYMS_SymbolNameMapLoose{
 } SYMS_SymbolNameMapLoose;
 
 ////////////////////////////////
-//~ NOTE(allen): String Cons Functions
+//~ allen: String Cons Functions
 
 SYMS_API SYMS_StringCons syms_string_cons_alloc(SYMS_Arena *arena, SYMS_U64 bucket_count);
 SYMS_API SYMS_String8    syms_string_cons(SYMS_Arena *arena, SYMS_StringCons *cons, SYMS_String8 string);
@@ -238,7 +238,7 @@ SYMS_API SYMS_DataIdxCons syms_data_idx_cons_alloc(SYMS_Arena *arena, SYMS_U64 b
 SYMS_API SYMS_U64         syms_data_idx_cons(SYMS_Arena *arena, SYMS_DataIdxCons *cons, SYMS_String8 data);
 
 ////////////////////////////////
-//~ NOTE(allen): U64 Set
+//~ allen: U64 Set
 
 SYMS_API SYMS_U64Set syms_u64_set_alloc(SYMS_Arena *arena, SYMS_U64 cap);
 SYMS_API SYMS_U64    syms_u64_set__bs(SYMS_U64Set *set, SYMS_U64 x);
@@ -246,7 +246,7 @@ SYMS_API SYMS_B32    syms_u64_set_insert(SYMS_U64Set *set, SYMS_U64 x);
 SYMS_API void        syms_u64_set_erase(SYMS_U64Set *set, SYMS_U64 x);
 
 ////////////////////////////////
-//~ NOTE(allen): 1D Spatial Mapping Functions (Overlaps Not Allowed)
+//~ allen: 1D Spatial Mapping Functions (Overlaps Not Allowed)
 
 //- lookups into spatial maps
 SYMS_API SYMS_U64          syms_spatial_map_1d_binary_search(SYMS_SpatialMap1D *map, SYMS_U64 x);
@@ -280,7 +280,7 @@ SYMS_API void              syms_spatial_map_1d_endpoint_sort(SYMS_1DEndPoint *en
 SYMS_API SYMS_B32          syms_spatial_map_1d_invariants(SYMS_SpatialMap1D *map);
 
 ////////////////////////////////
-//~ NOTE(allen): File Mapping Functions ({UnitID,FileID} -> String)
+//~ allen: File Mapping Functions ({UnitID,FileID} -> String)
 
 //- shared file id bucket definitions
 SYMS_API SYMS_U64            syms_file_id_2_name_map_hash(SYMS_UnitID uid, SYMS_FileID file_id);
@@ -301,7 +301,7 @@ SYMS_API void                syms_file_id_2_name_map_insert(SYMS_Arena *arena, S
 
 
 ////////////////////////////////
-//~ NOTE(allen): File Mapping Functions (String -> {UnitID,FileID})
+//~ allen: File Mapping Functions (String -> {UnitID,FileID})
 
 //- copying file maps
 SYMS_API SYMS_Name2FileIDMap syms_name_2_file_id_map_copy(SYMS_Arena *arena, SYMS_StringCons *cons_optional,
@@ -309,14 +309,14 @@ SYMS_API SYMS_Name2FileIDMap syms_name_2_file_id_map_copy(SYMS_Arena *arena, SYM
 
 //- constructing file maps
 SYMS_API SYMS_Name2FileIDMap syms_name_2_file_id_map_bake(SYMS_Arena *arena, SYMS_Name2FileIDMapLoose *loose);
-// NOTE(allen): Strings passed to this function should all be cons'ed in the same cons structure first.
+// allen: Strings passed to this function should all be cons'ed in the same cons structure first.
 SYMS_API void syms_name_2_file_id_map_loose_push(SYMS_Arena *arena, SYMS_Name2FileIDMapLoose *map,
                                                  SYMS_String8 name_cons,
                                                  SYMS_UnitID uid, SYMS_FileID file_id);
 
 
 ////////////////////////////////
-//~ NOTE(allen): ID Mapping Functions
+//~ allen: ID Mapping Functions
 
 //- copying id maps
 SYMS_API SYMS_IDMap   syms_id_map_copy(SYMS_Arena *arena, SYMS_IDMap *map);
@@ -330,7 +330,7 @@ SYMS_API void         syms_id_map_insert(SYMS_Arena *arena, SYMS_IDMap *map, SYM
 
 
 ////////////////////////////////
-//~ NOTE(allen): Symbol Name Mapping Structure (String -> Array(USID))
+//~ allen: Symbol Name Mapping Structure (String -> Array(USID))
 
 // TODO(allen): copying
 
@@ -346,7 +346,7 @@ SYMS_API SYMS_SymbolNameMap      syms_symbol_name_map_bake(SYMS_Arena *arena, SY
 
 
 ////////////////////////////////
-//~ NOTE(allen): Line Tables
+//~ allen: Line Tables
 
 //- lookups into line tables
 SYMS_API SYMS_U64  syms_line_index_from_voff__binary_search(SYMS_Line *lines, SYMS_U64 ender_index, SYMS_U64 voff);
@@ -360,12 +360,12 @@ SYMS_API SYMS_LineTable    syms_line_table_with_indexes_from_parse(SYMS_Arena *a
 
 
 ////////////////////////////////
-//~ NOTE(allen): Copies & Operators for Other Data Structures
+//~ allen: Copies & Operators for Other Data Structures
 
 SYMS_API SYMS_String8Array syms_string_array_copy(SYMS_Arena *arena, SYMS_StringCons *cons_optional,
                                                   SYMS_String8Array *array);
 
-SYMS_API SYMS_StrippedInfoArray syms_stripped_info_copy(SYMS_Arena *arena, SYMS_StrippedInfoArray *stripped);
+SYMS_API SYMS_LinkNameRecArray  syms_link_name_record_copy(SYMS_Arena *arena, SYMS_LinkNameRecArray *array);
 
 
 #endif //SYMS_DATA_STRUCTURES_H
diff --git a/Engine/Source/Runtime/SymsLib/syms/core/dwarf/syms_dwarf_expr.c b/Engine/Source/Runtime/SymsLib/syms/core/dwarf/syms_dwarf_expr.c
--- a/Engine/Source/Runtime/SymsLib/syms/core/dwarf/syms_dwarf_expr.c
+++ b/Engine/Source/Runtime/SymsLib/syms/core/dwarf/syms_dwarf_expr.c
@@ -382,6 +382,7 @@ syms_dw_expr__analyze_details(void *in_base, SYMS_U64Range in_range, SYMS_DwExpr
             result.flags |= SYMS_DwExprFlag_UsesObjectAddress;
           }break;
           
+          case SYMS_DwOp_GNU_PUSH_TLS_ADDRESS:
           case SYMS_DwOp_FORM_TLS_ADDRESS:
           {
             result.flags |= SYMS_DwExprFlag_UsesTLSAddress;
@@ -676,15 +677,29 @@ syms_dw_expr__eval(SYMS_Arena *arena_optional, void *expr_base, SYMS_U64Range ex
         {
           SYMS_U64 offset = 0;
           step_cursor += syms_based_range_read(base, range, step_cursor, 8, &offset);
-          if (config->text_section_base != 0){
-            SYMS_U64 x = *config->text_section_base + offset;
-            syms_dw_expr__stack_push(scratch.arena, &stack, x);
+
+          // NOTE(nick): Earlier versions of GCC emit TLS offset with SYMS_DwOp_ADDR.
+          SYMS_B32 is_text_relative;
+          {
+            SYMS_U8 next_op = 0;
+            syms_based_range_read_struct(base, range, step_cursor, &next_op);
+            is_text_relative = (next_op != SYMS_DwOp_GNU_PUSH_TLS_ADDRESS);
           }
-          else{
-            stashed_loc.kind = SYMS_DwSimpleLocKind_Fail;
-            stashed_loc.fail_kind = SYMS_DwLocFailKind_MissingTextBase;
-            goto finish;
+
+          SYMS_U64 addr = offset;
+
+          if (is_text_relative){
+            if (config->text_section_base != 0){
+              addr += *config->text_section_base;
+            }
+            else{
+              stashed_loc.kind = SYMS_DwSimpleLocKind_Fail;
+              stashed_loc.fail_kind = SYMS_DwLocFailKind_MissingTextBase;
+              goto finish;
+            }
           }
+
+          syms_dw_expr__stack_push(scratch.arena, &stack, addr);
         }break;
         
         case SYMS_DwOp_CONSTU:
@@ -890,13 +905,13 @@ syms_dw_expr__eval(SYMS_Arena *arena_optional, void *expr_base, SYMS_U64Range ex
           }
         }break;
         
+        // NOTE: pop offset from stack, convert it to TLS address, then push it back.
+        case SYMS_DwOp_GNU_PUSH_TLS_ADDRESS:
         case SYMS_DwOp_FORM_TLS_ADDRESS:
         {
           SYMS_S64 s = (SYMS_S64)syms_dw_expr__stack_pop(&stack);
-          
+
           if (config->tls_address != 0){
-            // TODO(allen): Is this what the spec means??
-            // Or do we need to interpret the value from the stack in some other way?
             SYMS_U64 x = *config->tls_address + s;
             syms_dw_expr__stack_push(scratch.arena, &stack, x);
           }
diff --git a/Engine/Source/Runtime/SymsLib/syms/core/dwarf/syms_dwarf_parser.c b/Engine/Source/Runtime/SymsLib/syms/core/dwarf/syms_dwarf_parser.c
--- a/Engine/Source/Runtime/SymsLib/syms/core/dwarf/syms_dwarf_parser.c
+++ b/Engine/Source/Runtime/SymsLib/syms/core/dwarf/syms_dwarf_parser.c
@@ -1076,6 +1076,14 @@ syms_dw_uid_from_foff(SYMS_DwDbgAccel *dbg, SYMS_U64 foff)
 
 //- rjf: important DWARF section base/range accessors
 
+SYMS_API SYMS_B32
+syms_dw_sec_is_present(SYMS_DwDbgAccel *dbg, SYMS_DwSectionKind kind)
+{
+  SYMS_U64Range *range = &dbg->section_map[kind].range;
+  SYMS_B32 result = (range->max > range->min);
+  return(result);
+}
+
 SYMS_API void *
 syms_dw_sec_base_from_dbg(SYMS_String8 data, SYMS_DwDbgAccel *dbg, SYMS_DwSectionKind kind)
 {
@@ -5602,7 +5610,7 @@ SYMS_API SYMS_DwMapAccel*
 syms_dw_type_map_from_dbg(SYMS_Arena *arena, SYMS_String8 data, SYMS_DwDbgAccel *dbg)
 {
   SYMS_DwMapAccel *result = (SYMS_DwMapAccel *)&syms_format_nil;
-  if(syms_dw_sec_base_from_dbg(data, dbg, SYMS_DwSectionKind_PubTypes))
+  if(syms_dw_sec_is_present(dbg, SYMS_DwSectionKind_PubTypes))
   {
     result = syms_push_array(arena, SYMS_DwMapAccel, 1);
     syms_memzero_struct(result);
@@ -5616,7 +5624,7 @@ SYMS_API SYMS_DwMapAccel*
 syms_dw_image_symbol_map_from_dbg(SYMS_Arena *arena, SYMS_String8 data, SYMS_DwDbgAccel *dbg)
 {
   SYMS_DwMapAccel *result = (SYMS_DwMapAccel *)&syms_format_nil;
-  if(syms_dw_sec_base_from_dbg(data, dbg, SYMS_DwSectionKind_PubNames))
+  if(syms_dw_sec_is_present(dbg, SYMS_DwSectionKind_PubNames))
   {
     result = syms_push_array(arena, SYMS_DwMapAccel, 1);
     syms_memzero_struct(result);
diff --git a/Engine/Source/Runtime/SymsLib/syms/core/dwarf/syms_dwarf_parser.h b/Engine/Source/Runtime/SymsLib/syms/core/dwarf/syms_dwarf_parser.h
--- a/Engine/Source/Runtime/SymsLib/syms/core/dwarf/syms_dwarf_parser.h
+++ b/Engine/Source/Runtime/SymsLib/syms/core/dwarf/syms_dwarf_parser.h
@@ -546,6 +546,13 @@ struct SYMS_DwMapAccel
   SYMS_DwPubStringsTable tbl;
 };
 
+typedef struct SYMS_DwLinkMapAccel SYMS_DwLinkMapAccel;
+struct SYMS_DwLinkMapAccel
+{
+  SYMS_FileFormat format;
+  // TODO(allen): 
+};
+
 SYMS_C_LINKAGE_BEGIN //-
 
 ////////////////////////////////
@@ -590,6 +597,7 @@ SYMS_API SYMS_U64          syms_dw_default_vbase_from_dbg(SYMS_DwDbgAccel *dbg);
 SYMS_API SYMS_UnitID syms_dw_uid_from_foff(SYMS_DwDbgAccel *dbg, SYMS_U64 foff);
 
 //- rjf: important DWARF section base/range accessors
+SYMS_API SYMS_B32      syms_dw_sec_is_present(SYMS_DwDbgAccel *dbg, SYMS_DwSectionKind kind);
 SYMS_API void*         syms_dw_sec_base_from_dbg(SYMS_String8 data, SYMS_DwDbgAccel *dbg,
                                                  SYMS_DwSectionKind kind);
 SYMS_API SYMS_U64Range syms_dw_sec_range_from_dbg(SYMS_DwDbgAccel *dbg, SYMS_DwSectionKind kind);
diff --git a/Engine/Source/Runtime/SymsLib/syms/core/extras/syms_aggregate_proc_map.c b/Engine/Source/Runtime/SymsLib/syms/core/extras/syms_aggregate_proc_map.c
--- a/Engine/Source/Runtime/SymsLib/syms/core/extras/syms_aggregate_proc_map.c
+++ b/Engine/Source/Runtime/SymsLib/syms/core/extras/syms_aggregate_proc_map.c
@@ -1,3 +1,4 @@
+// Copyright Epic Games, Inc. All Rights Reserved.
 ////////////////////////////////
 // NOTE(allen): Syms Procedure Aggregator Function
 
diff --git a/Engine/Source/Runtime/SymsLib/syms/core/group/syms_group.c b/Engine/Source/Runtime/SymsLib/syms/core/group/syms_group.c
--- a/Engine/Source/Runtime/SymsLib/syms/core/group/syms_group.c
+++ b/Engine/Source/Runtime/SymsLib/syms/core/group/syms_group.c
@@ -4,7 +4,7 @@
 #define SYMS_GROUP_C
 
 ////////////////////////////////
-// NOTE(allen): Syms Group Setup Functions
+// allen: Syms Group Setup Functions
 
 SYMS_API SYMS_Group*
 syms_group_alloc(void){
@@ -79,6 +79,7 @@ syms_group_init(SYMS_Group *group, SYMS_ParseBundle *params){
   SYMS_UnitSetAccel *unit_set = syms_unit_set_accel_from_dbg(group->arena, dbg_data, dbg);
   SYMS_U64 unit_count = syms_unit_count_from_set(unit_set);
   SYMS_MapAccel *type_map = syms_type_map_from_dbg(group->arena, dbg_data, dbg);
+  SYMS_MapAccel *symbol_map = syms_unmangled_symbol_map_from_dbg(group->arena, dbg_data, dbg);
   
   //- init lanes to single-threaded setup
   group->lane_arenas = &group->arena;
@@ -96,7 +97,8 @@ syms_group_init(SYMS_Group *group, SYMS_ParseBundle *params){
   group->sec_info_array = sec_info_array;
   group->unit_set = unit_set;
   group->unit_count = unit_count;
-  group->type_map = type_map;
+  group->type_mau.map = type_map;
+  group->symbol_mau.map = symbol_map;
   
   //- setup basic caches
   group->unit_cache_flags = syms_push_array_zero(group->arena, SYMS_GroupUnitCacheFlags, unit_count);
@@ -154,6 +156,7 @@ syms_group_parse_all_top_level(SYMS_Group *group){
   
   syms_group_unit_ranges(group);
   syms_group_type_map(group);
+  syms_group_unmangled_symbol_map(group);
   
   syms_group_sec_map_v(group);
   syms_group_sec_map_f(group);
@@ -213,7 +216,7 @@ syms_group_get_lane_arena(SYMS_Group *group){
 }
 
 ////////////////////////////////
-// NOTE(allen): Syms Group Getters
+// allen: Syms Group Getters
 
 SYMS_API SYMS_String8
 syms_group_bin_data(SYMS_Group *group){
@@ -527,35 +530,138 @@ syms_group_file_name_from_id(SYMS_Group *group, SYMS_UnitID uid, SYMS_FileID fil
   return(result);
 }
 
-SYMS_API SYMS_MapAndUnit
+SYMS_API SYMS_MapAndUnit*
 syms_group_type_map(SYMS_Group *group){
   if (!group->type_map_unit_is_filled){
-    SYMS_UnitID type_map_uid = syms_partner_uid_from_map(group->type_map);
+    SYMS_UnitID type_map_uid = syms_partner_uid_from_map(group->type_mau.map);
     SYMS_UnitAccel *type_map_unit = syms_group_unit_from_uid(group, type_map_uid);
-    group->type_map_unit = type_map_unit;
+    group->type_mau.unit = type_map_unit;
     group->type_map_unit_is_filled = syms_true;
   }
-  SYMS_MapAndUnit result = {0};
-  result.map = group->type_map;
-  result.unit = group->type_map_unit;
+  SYMS_MapAndUnit *result = &group->type_mau;
   return(result);
 }
 
-SYMS_API SYMS_StrippedInfoArray
-syms_group_stripped_info(SYMS_Group *group){
-  if (!group->stripped_info_is_filled){
-    SYMS_UnitID pub_uid = syms_uid_collated_public_symbols_from_set(group->unit_set);
-    SYMS_UnitAccel *pub = syms_group_unit_from_uid(group, pub_uid);
-    group->stripped_info = syms_stripped_from_unit(group->arena, group->dbg_data, group->dbg, pub);
-    group->stripped_info_is_filled = syms_true;
+SYMS_API SYMS_MapAndUnit*
+syms_group_unmangled_symbol_map(SYMS_Group *group){
+  if (!group->symbol_map_unit_is_filled){
+    SYMS_UnitID symbol_map_uid = syms_partner_uid_from_map(group->symbol_mau.map);
+    SYMS_UnitAccel *symbol_map_unit = syms_group_unit_from_uid(group, symbol_map_uid);
+    group->symbol_mau.unit = symbol_map_unit;
+    group->symbol_map_unit_is_filled = syms_true;
   }
-  SYMS_StrippedInfoArray result = group->stripped_info;
+  SYMS_MapAndUnit *result = &group->symbol_mau;
   return(result);
 }
 
+SYMS_API SYMS_LinkNameRecArray
+syms_group_link_name_records(SYMS_Group *group){
+  if (!group->link_name_record_array_is_filled){
+    SYMS_UnitID link_uid = syms_link_names_uid(group->dbg);
+    SYMS_UnitAccel *unit = syms_group_unit_from_uid(group, link_uid);
+    SYMS_LinkNameRecArray array = syms_link_name_array_from_unit(group->arena, group->dbg_data, group->dbg, unit);
+    group->link_name_record_array = array;
+    
+    // sort
+    syms_group__link_names_sort_in_place(array.recs, array.count);
+    
+    group->link_name_record_array_is_filled = syms_true;
+  }
+  SYMS_LinkNameRecArray result = group->link_name_record_array;
+  return(result);
+}
+
+SYMS_API SYMS_LinkMapAccel*
+syms_group_link_name_map(SYMS_Group *group){
+  if (!group->link_map_is_built){
+    group->link_map = syms_link_map_from_dbg(group->arena, group->dbg_data, group->dbg);
+    group->link_map_is_built = !group->link_map_is_built;
+  }
+  SYMS_LinkMapAccel *result = group->link_map;
+  return(result);
+}
+
+SYMS_API SYMS_UnitAccel*
+syms_group_link_name_unit(SYMS_Group *group){
+  if (!group->link_name_unit_is_filled){
+    SYMS_UnitID uid = syms_link_names_uid(group->dbg);
+    group->link_name_unit = syms_group_unit_from_uid(group, uid);
+    group->link_name_unit_is_filled = syms_true;
+  }
+  SYMS_UnitAccel *result = group->link_name_unit;
+  return(result);
+}
+
+SYMS_API void
+syms_group__link_names_sort_in_place(SYMS_LinkNameRec *recs, SYMS_U64 count){
+  SYMS_ArenaTemp scratch = syms_get_scratch(0, 0);
+  
+  // setup stack with full range
+  SYMS_SortNode *stack = 0;
+  SYMS_SortNode *free_stack = 0;
+  syms_sort_node_push(scratch.arena, &stack, &free_stack, 0, count);
+  
+  // sort loop
+  for (; stack != 0;){
+    SYMS_SortNode *node = stack;
+    SYMS_StackPop(stack);
+    
+    SYMS_U64 first = node->first;
+    SYMS_U64 opl = node->opl;
+    
+    SYMS_StackPush(free_stack, node);
+    
+    SYMS_U64 node_count = (opl - first);
+    if (node_count > 2){
+      SYMS_U64 last = opl - 1;
+      
+      // swap key to last
+      SYMS_U64 mid = (first + opl)/2;
+      SYMS_Swap(SYMS_LinkNameRec, recs[mid], recs[last]);
+      
+      // partition
+      SYMS_B32 equal_send_left = syms_false;
+      SYMS_U64 key = recs[last].voff;
+      SYMS_U64 j = first;
+      for (SYMS_U64 i = first; i < last; i += 1){
+        SYMS_B32 send_left = (recs[i].voff < key);
+        if (!send_left && recs[i].voff == key){
+          send_left = equal_send_left;
+          equal_send_left = !equal_send_left;
+        }
+        if (send_left){
+          if (j != i){
+            SYMS_Swap(SYMS_LinkNameRec, recs[i], recs[j]);
+          }
+          j += 1;
+        }
+      }
+      
+      // swap last to pivot
+      SYMS_Swap(SYMS_LinkNameRec, recs[j], recs[last]);
+      
+      // recurse
+      SYMS_U64 pivot = j;
+      if (pivot - first > 1){
+        syms_sort_node_push(scratch.arena, &stack, &free_stack, first, pivot);
+      }
+      if (opl - (pivot + 1) > 1){
+        syms_sort_node_push(scratch.arena, &stack, &free_stack, pivot + 1, opl);
+      }
+    }
+    else if (node_count == 2){
+      if (recs[first].voff > recs[first + 1].voff){
+        SYMS_Swap(SYMS_LinkNameRec, recs[first], recs[first + 1]);
+      }
+    }
+  }
+  
+  syms_release_scratch(scratch);
+}
+
 
 ////////////////////////////////
-// NOTE(allen): Syms Group Address Mapping Functions
+// allen: Syms Group Address Mapping Functions
 
 SYMS_API SYMS_U64
 syms_group_sec_number_from_voff__linear_scan(SYMS_Group *group, SYMS_U64 voff){
@@ -1211,23 +1317,23 @@ syms_group_type_map_from_uid(SYMS_Group *group, SYMS_UnitID uid){
 
 
 SYMS_API SYMS_SpatialMap1D*
-syms_group_stripped_info_map(SYMS_Group *group){
+syms_group_link_name_spatial_map(SYMS_Group *group){
   SYMS_ASSERT_PARANOID(syms_thread_lane == 0);
-  SYMS_ProfBegin("syms_group_stripped_info_map");
-  if (!group->stripped_info_map_is_built){
+  SYMS_ProfBegin("syms_group_link_name_spatial_map");
+  if (!group->link_name_spatial_map_is_built){
     SYMS_Arena *arena = group->arena;
     
     //- fill spatial map array
-    SYMS_StrippedInfoArray stripped_info = syms_group_stripped_info(group);
-    SYMS_U64 guess_count = stripped_info.count;
+    SYMS_LinkNameRecArray link_name_rec_array = syms_group_link_name_records(group);
+    SYMS_U64 guess_count = link_name_rec_array.count;
     SYMS_SpatialMap1DRange *ranges = syms_push_array(arena, SYMS_SpatialMap1DRange, guess_count);
     SYMS_SpatialMap1DRange *range_ptr = ranges;
-    SYMS_StrippedInfo *stripped_ptr = stripped_info.info;
-    for (SYMS_U64 i = 0; i < guess_count; i += 1, stripped_ptr += 1){
-      SYMS_U64 min = stripped_ptr->voff;
+    SYMS_LinkNameRec *rec_ptr = link_name_rec_array.recs;
+    for (SYMS_U64 i = 0; i < guess_count; i += 1, rec_ptr += 1){
+      SYMS_U64 min = rec_ptr->voff;
       SYMS_U64 max = SYMS_U64_MAX;
       if (i + 1 < guess_count){
-        max = (stripped_ptr + 1)->voff;
+        max = (rec_ptr + 1)->voff;
       }
       SYMS_ASSERT(min <= max);
       if (min < max){
@@ -1246,13 +1352,13 @@ syms_group_stripped_info_map(SYMS_Group *group){
     SYMS_SpatialMap1D map = {ranges, actual_count};
     
     //- save to group
-    group->stripped_info_map_is_built = syms_true;
-    group->stripped_info_map = map;
+    group->link_name_spatial_map_is_built = syms_true;
+    group->link_name_spatial_map = map;
     
     SYMS_ASSERT_PARANOID(syms_spatial_map_1d_invariants(&map));
   }
   
-  SYMS_SpatialMap1D *result = &group->stripped_info_map;
+  SYMS_SpatialMap1D *result = &group->link_name_spatial_map;
   SYMS_ProfEnd();
   return(result);
 }
@@ -1461,15 +1567,70 @@ syms_line_to_addr_line_sort__rec(SYMS_FileToLineToAddrLooseLine **array, SYMS_U6
 }
 
 ////////////////////////////////
-// NOTE(allen): Syms Group Type Graph
+// allen: Syms Group Name Mapping Functions
+
+SYMS_API SYMS_USID
+syms_group_usid_from_unmangled_name(SYMS_Group *group, SYMS_String8 name){
+  SYMS_USID result = {0};
+  SYMS_ArenaTemp scratch = syms_get_scratch(0, 0);
+  SYMS_MapAndUnit *map = syms_group_unmangled_symbol_map(group);
+  if (syms_accel_is_good(map->map)){
+    SYMS_USIDList list = syms_usid_list_from_string(scratch.arena, group->dbg_data, group->dbg, map, name);
+    if (list.first != 0){
+      result = list.first->usid;
+    }
+  }
+  syms_release_scratch(scratch);
+  return(result);
+}
 
-SYMS_API SYMS_TypeGraph* syms_group_type_graph(SYMS_Group *group){
+SYMS_API SYMS_USIDList
+syms_group_all_usid_from_unmangled_name(SYMS_Arena *arena, SYMS_Group *group, SYMS_String8 name){
+  SYMS_USIDList result = {0};
+  SYMS_MapAndUnit *map = syms_group_unmangled_symbol_map(group);
+  if (syms_accel_is_good(map->map)){
+    result = syms_usid_list_from_string(arena, group->dbg_data, group->dbg, map, name);
+  }
+  return(result);
+}
+
+SYMS_API SYMS_U64
+syms_group_voff_from_link_name(SYMS_Group *group, SYMS_String8 name){
+  SYMS_LinkMapAccel *map = syms_group_link_name_map(group);
+  SYMS_UnitAccel *unit = syms_group_link_name_unit(group);
+  SYMS_U64 result = syms_voff_from_link_name(group->dbg_data, group->dbg, map, unit, name);
+  return(result);
+}
+
+SYMS_API SYMS_ResolvedLine
+syms_group_resolved_location_from_link_name(SYMS_Group *group, SYMS_String8 name){
+  SYMS_U64 voff = syms_group_voff_from_link_name(group, name);
+  SYMS_UnitID uid = syms_group_uid_from_voff__accelerated(group, voff);
+  SYMS_Line line = syms_group_line_from_uid_voff__accelerated(group, uid, voff);
+  SYMS_SrcCoord *src_coord = &line.src_coord;
+  SYMS_String8 file_name = syms_group_file_name_from_id(group, uid, src_coord->file_id);
+  
+  SYMS_ResolvedLine result = {0};
+  result.file_name = file_name;
+  result.line = src_coord->line;
+  result.col = src_coord->col;
+  result.voff = voff;
+  return(result);
+}
+
+
+////////////////////////////////
+// allen: Syms Group Type Graph
+
+SYMS_API SYMS_TypeGraph*
+syms_group_type_graph(SYMS_Group *group){
   SYMS_TypeGraph *result = &group->type_graph;
   return(result);
 }
 
+
 ////////////////////////////////
-// NOTE(allen): Syms Group Varaible Address Mapping Functions
+// allen: Syms Group Varaible Address Mapping Functions
 
 SYMS_API SYMS_SymbolID
 syms_group_var_sid_from_uid_voff__linear_scan(SYMS_TypeGraph *graph, SYMS_Group *group,
@@ -1577,7 +1738,7 @@ syms_group_var_sid_from_uid_voff__accelerated(SYMS_TypeGraph *graph, SYMS_Group
 
 
 ////////////////////////////////
-// NOTE(allen): Syms Group Type Graph Functions
+// allen: Syms Group Type Graph Functions
 
 SYMS_API SYMS_TypeNode*
 syms_group_type_from_usid(SYMS_TypeGraph *graph, SYMS_Group *group, SYMS_USID usid){
@@ -1650,9 +1811,9 @@ syms_group_type_from_usid__rec(SYMS_TypeGraph *graph, SYMS_Group *group, SYMS_US
               SYMS_TypeKind match_kind = syms_type_kind_main_from_fwd(type_info.kind);
               
               SYMS_String8 name = syms_group_symbol_name_from_sid(scratch.arena, group, unit, usid.sid);
-              SYMS_MapAndUnit type_map = syms_group_type_map(group);
-              SYMS_USIDList matches = syms_symbol_list_from_string(scratch.arena, group->dbg_data, group->dbg,
-                                                                   type_map, name);
+              SYMS_MapAndUnit *type_map = syms_group_type_map(group);
+              SYMS_USIDList matches = syms_usid_list_from_string(scratch.arena, group->dbg_data, group->dbg,
+                                                                 type_map, name);
               
               SYMS_USID match_usid = {0};
               for (SYMS_USIDNode *node = matches.first;
@@ -1807,8 +1968,7 @@ syms_group_type_from_usid__rec(SYMS_TypeGraph *graph, SYMS_Group *group, SYMS_US
         }
         
         if (!pre_inserted){
-          syms_type_usid_buckets_insert(graph->arena, &graph->type_usid_buckets,
-                                        usid, new_type);
+          syms_type_usid_buckets_insert(graph->arena, &graph->type_usid_buckets, usid, new_type);
         }
       }
       
@@ -1856,8 +2016,7 @@ syms_group_type_size_from_usid(SYMS_TypeGraph *graph, SYMS_Group *group, SYMS_US
   return(result);
 }
 
-SYMS_API SYMS_TypeChain syms_group_artificial_types_from_name(SYMS_Group *group,
-                                                              SYMS_String8 name){
+SYMS_API SYMS_TypeChain syms_group_artificial_types_from_name(SYMS_Group *group, SYMS_String8 name){
   SYMS_TypeChain result = syms_type_from_name(&group->type_graph, name);
   return(result);
 }
@@ -1865,15 +2024,10 @@ SYMS_API SYMS_TypeChain syms_group_artificial_types_from_name(SYMS_Group *group,
 SYMS_API SYMS_USIDList
 syms_group_type_list_from_name_accelerated(SYMS_Arena *arena, SYMS_Group *group, SYMS_String8 name){
   SYMS_USIDList result = {0};
-  
-  SYMS_MapAndUnit map_unit = syms_group_type_map(group);
-  if (syms_accel_is_good(map_unit.map)){
-    SYMS_String8 dbg_data = group->dbg_data;
-    SYMS_DbgAccel *dbg = group->dbg;
-    result = syms_symbol_list_from_string(arena, dbg_data, dbg,
-                                          map_unit, name);
+  SYMS_MapAndUnit *map_unit = syms_group_type_map(group);
+  if (syms_accel_is_good(map_unit->map)){
+    result = syms_usid_list_from_string(arena, group->dbg_data, group->dbg, map_unit, name);
   }
-  
   return(result);
 }
 
@@ -2024,7 +2178,7 @@ syms_type_enum_members_from_type(SYMS_Group *group,
 }
 
 ////////////////////////////////
-//~ NOTE(allen): Syms File Mapp
+//~ allen: Syms File Mapp
 
 SYMS_API SYMS_Name2FileIDMap*
 syms_group_file_map(SYMS_Group *group){
diff --git a/Engine/Source/Runtime/SymsLib/syms/core/group/syms_group.h b/Engine/Source/Runtime/SymsLib/syms/core/group/syms_group.h
--- a/Engine/Source/Runtime/SymsLib/syms/core/group/syms_group.h
+++ b/Engine/Source/Runtime/SymsLib/syms/core/group/syms_group.h
@@ -5,7 +5,7 @@
 #define SYMS_GROUP_H
 
 ////////////////////////////////
-// NOTE(allen): Syms Line Mapping Structures
+//~ allen: Syms Line Mapping Structures
 
 typedef struct SYMS_LineToAddrMap{
   SYMS_U64Range *ranges;
@@ -51,7 +51,7 @@ typedef struct SYMS_FileToLineToAddrLoose{
 } SYMS_FileToLineToAddrLoose;
 
 ////////////////////////////////
-// NOTE(allen): Syms Group Types
+//~ allen: Syms Group Types
 
 typedef SYMS_U32 SYMS_GroupUnitCacheFlags;
 enum{
@@ -90,7 +90,6 @@ typedef struct SYMS_Group{
   SYMS_SecInfoArray sec_info_array;
   SYMS_UnitSetAccel *unit_set;
   SYMS_U64 unit_count;
-  SYMS_MapAccel *type_map;
   
   //- basic section caches
   SYMS_String8 *sec_names;
@@ -121,36 +120,45 @@ typedef struct SYMS_Group{
   //- one-time fills/builds
   SYMS_B8 unit_ranges_is_filled;
   SYMS_B8 type_map_unit_is_filled;
+  SYMS_B8 symbol_map_unit_is_filled;
+  SYMS_B8 mangled_symbol_map_unit_is_filled;
   SYMS_B8 sec_map_v_is_built;
   SYMS_B8 sec_map_f_is_built;
   SYMS_B8 unit_map_is_built;
   SYMS_B8 name_2_file_id_map_is_built;
-  SYMS_B8 stripped_info_is_filled;
-  SYMS_B8 stripped_info_map_is_built;
+  SYMS_B8 link_name_record_array_is_filled;
+  SYMS_B8 link_name_spatial_map_is_built;
+  SYMS_B8 link_map_is_built;
+  SYMS_B8 link_name_unit_is_filled;
   
   SYMS_UnitRangeArray unit_ranges;
-  SYMS_UnitAccel *type_map_unit;
   SYMS_SpatialMap1D sec_map_v;
   SYMS_SpatialMap1D sec_map_f;
   SYMS_SpatialMap1D unit_map;
   SYMS_Name2FileIDMap name_2_file_id_map;
-  SYMS_StrippedInfoArray stripped_info;
-  SYMS_SpatialMap1D stripped_info_map;
+  SYMS_LinkNameRecArray link_name_record_array;
+  SYMS_SpatialMap1D link_name_spatial_map;
+  SYMS_LinkMapAccel *link_map;
+  SYMS_UnitAccel *link_name_unit;
+  
+  //- map & units
+  SYMS_MapAndUnit type_mau;
+  SYMS_MapAndUnit symbol_mau;
 } SYMS_Group;
 
 ////////////////////////////////
-// NOTE(allen): Data Structure Nils
+//~ allen: Data Structure Nils
 
 SYMS_READ_ONLY SYMS_GLOBAL SYMS_SymbolIDArray syms_sid_array_nil = {0};
 SYMS_READ_ONLY SYMS_GLOBAL SYMS_String8Array syms_string_array_nil = {0};
-SYMS_READ_ONLY SYMS_GLOBAL SYMS_LineParseOut syms_line_parse_nil = {{0}};
+SYMS_READ_ONLY SYMS_GLOBAL SYMS_LineParseOut syms_line_parse_nil = {0};
 SYMS_READ_ONLY SYMS_GLOBAL SYMS_SpatialMap1D syms_spatial_map_1d_nil = {0};
 SYMS_READ_ONLY SYMS_GLOBAL SYMS_LineToAddrMap syms_line_to_addr_map_nil = {0};
 SYMS_READ_ONLY SYMS_GLOBAL SYMS_TypeMemberArray syms_type_member_array_nil = {0};
 SYMS_READ_ONLY SYMS_GLOBAL SYMS_EnumInfoArray syms_enum_info_array_nil = {0};
 
 ////////////////////////////////
-// NOTE(allen): Syms Group Setup Functions
+//~ allen: Syms Group Setup Functions
 
 SYMS_API SYMS_Group* syms_group_alloc(void);
 SYMS_API void        syms_group_release(SYMS_Group *group);
@@ -166,7 +174,7 @@ SYMS_API void        syms_group_end_multilane(SYMS_Group *group);
 SYMS_API SYMS_Arena* syms_group_get_lane_arena(SYMS_Group *group);
 
 ////////////////////////////////
-// NOTE(allen): Syms Group Getters & Cache Accessors
+//~ allen: Syms Group Getters & Cache Accessors
 
 // TODO(allen): better sorting plan here.
 
@@ -208,13 +216,16 @@ SYMS_API SYMS_String8       syms_group_symbol_name_from_sid(SYMS_Arena *arena, S
 
 SYMS_API SYMS_String8       syms_group_file_name_from_id(SYMS_Group *group, SYMS_UnitID uid, SYMS_FileID file_id);
 
-SYMS_API SYMS_MapAndUnit    syms_group_type_map(SYMS_Group *group);
+SYMS_API SYMS_MapAndUnit*   syms_group_type_map(SYMS_Group *group);
+SYMS_API SYMS_MapAndUnit*   syms_group_unmangled_symbol_map(SYMS_Group *group);
 
-SYMS_API SYMS_StrippedInfoArray syms_group_stripped_info(SYMS_Group *group);
+SYMS_API SYMS_LinkNameRecArray syms_group_link_name_records(SYMS_Group *group);
+SYMS_API SYMS_LinkMapAccel*    syms_group_link_name_map(SYMS_Group *group);
+SYMS_API SYMS_UnitAccel*       syms_group_link_name_unit(SYMS_Group *group);
 
 
 ////////////////////////////////
-// NOTE(allen): Syms Group Address Mapping Functions
+//~ allen: Syms Group Address Mapping Functions
 
 //- linear scan versions
 SYMS_API SYMS_U64      syms_group_sec_number_from_voff__linear_scan(SYMS_Group *group, SYMS_U64 voff);
@@ -241,20 +252,16 @@ SYMS_API SYMS_SpatialMap1D* syms_group_unit_map(SYMS_Group *group);
 
 // thread safe (with lanes equipped to group)
 SYMS_API SYMS_SpatialMap1D* syms_group_proc_map_from_uid(SYMS_Group *group, SYMS_UnitID uid);
-SYMS_API SYMS_SpatialMap1D* syms_group_line_sequence_map_from_uid(SYMS_Group *group,
-                                                                  SYMS_UnitID uid);
+SYMS_API SYMS_SpatialMap1D* syms_group_line_sequence_map_from_uid(SYMS_Group *group, SYMS_UnitID uid);
 
-SYMS_API void                syms_group_fetch_line_to_addr_maps_from_uid(SYMS_Group *group,
-                                                                         SYMS_UnitID uid);
-SYMS_API SYMS_LineToAddrMap*
-syms_group_line_to_addr_map_from_uid_file_id(SYMS_Group *group,
-                                             SYMS_UnitID uid,
-                                             SYMS_FileID file_id);
+SYMS_API void                syms_group_fetch_line_to_addr_maps_from_uid(SYMS_Group *group, SYMS_UnitID uid);
+SYMS_API SYMS_LineToAddrMap* syms_group_line_to_addr_map_from_uid_file_id(SYMS_Group *group, SYMS_UnitID uid,
+                                                                          SYMS_FileID file_id);
 
-SYMS_API SYMS_SymbolNameMap*
-syms_group_type_map_from_uid(SYMS_Group *group, SYMS_UnitID uid);
+SYMS_API SYMS_SymbolNameMap* syms_group_type_map_from_uid(SYMS_Group *group, SYMS_UnitID uid);
 
-SYMS_API SYMS_SpatialMap1D* syms_group_stripped_info_map(SYMS_Group *group);
+SYMS_API SYMS_SpatialMap1D*  syms_group_link_name_spatial_map(SYMS_Group *group);
+SYMS_API void                syms_group__link_names_sort_in_place(SYMS_LinkNameRec *recs, SYMS_U64 count);
 
 //- accelerated versions
 SYMS_API SYMS_U64      syms_group_sec_number_from_voff__accelerated(SYMS_Group *group, SYMS_U64 voff);
@@ -280,59 +287,56 @@ SYMS_API void          syms_line_to_addr_line_sort(SYMS_FileToLineToAddrLooseLin
 SYMS_API void          syms_line_to_addr_line_sort__rec(SYMS_FileToLineToAddrLooseLine **array, SYMS_U64 count);
 
 ////////////////////////////////
-// NOTE(allen): Syms Group Type Graph
+//~ allen: Syms Group Name Mapping Functions
+
+SYMS_API SYMS_USID     syms_group_usid_from_unmangled_name(SYMS_Group *group, SYMS_String8 name);
+SYMS_API SYMS_USIDList syms_group_all_usid_from_unmangled_name(SYMS_Arena *arena, SYMS_Group *group,
+                                                               SYMS_String8 name);
+
+SYMS_API SYMS_U64      syms_group_voff_from_link_name(SYMS_Group *group, SYMS_String8 name);
+SYMS_API SYMS_ResolvedLine syms_group_resolved_location_from_link_name(SYMS_Group *group, SYMS_String8 name);
+
+////////////////////////////////
+//~ allen: Syms Group Type Graph
 
 SYMS_API SYMS_TypeGraph* syms_group_type_graph(SYMS_Group *group);
 
 ////////////////////////////////
-// NOTE(allen): Syms Group Varaible Address Mapping Functions
+//~ allen: Syms Group Varaible Address Mapping Functions
 
 // these are seperated because they require type info which
 // cannot be parallelized like the others
 
 
-SYMS_API SYMS_SymbolID syms_group_var_sid_from_uid_voff__linear_scan(SYMS_TypeGraph *graph,
-                                                                     SYMS_Group *group,
-                                                                     SYMS_UnitID uid,
-                                                                     SYMS_U64 voff);
+SYMS_API SYMS_SymbolID syms_group_var_sid_from_uid_voff__linear_scan(SYMS_TypeGraph *graph, SYMS_Group *group,
+                                                                     SYMS_UnitID uid, SYMS_U64 voff);
 
-SYMS_API SYMS_SpatialMap1D* syms_group_var_map_from_uid(SYMS_TypeGraph *graph,
-                                                        SYMS_Group *group, SYMS_UnitID uid);
+SYMS_API SYMS_SpatialMap1D* syms_group_var_map_from_uid(SYMS_TypeGraph *graph, SYMS_Group *group, SYMS_UnitID uid);
 
-SYMS_API SYMS_SymbolID syms_group_var_sid_from_uid_voff__accelerated(SYMS_TypeGraph *graph,
-                                                                     SYMS_Group *group,
-                                                                     SYMS_UnitID uid,
-                                                                     SYMS_U64 voff);
+SYMS_API SYMS_SymbolID syms_group_var_sid_from_uid_voff__accelerated(SYMS_TypeGraph *graph, SYMS_Group *group,
+                                                                     SYMS_UnitID uid, SYMS_U64 voff);
 
 
 ////////////////////////////////
-// NOTE(allen): Syms Group Type Info Functions
+//~ allen: Syms Group Type Info Functions
 
-SYMS_API SYMS_TypeNode* syms_group_type_from_usid(SYMS_TypeGraph *graph, SYMS_Group *group,
-                                                  SYMS_USID usid);
-SYMS_API SYMS_TypeNode* syms_group_type_from_usid__rec(SYMS_TypeGraph *graph,
-                                                       SYMS_Group *group, SYMS_USID usid);
+SYMS_API SYMS_TypeNode* syms_group_type_from_usid(SYMS_TypeGraph *graph, SYMS_Group *group, SYMS_USID usid);
+SYMS_API SYMS_TypeNode* syms_group_type_from_usid__rec(SYMS_TypeGraph *graph, SYMS_Group *group, SYMS_USID usid);
 
-SYMS_API SYMS_TypeChain syms_group_artificial_types_from_name(SYMS_Group *group,
-                                                              SYMS_String8 name);
+SYMS_API SYMS_TypeChain syms_group_artificial_types_from_name(SYMS_Group *group, SYMS_String8 name);
 
-SYMS_API SYMS_USIDList syms_group_type_list_from_name_accelerated(SYMS_Arena *arena,
-                                                                  SYMS_Group *group,
+// TODO(allen): sort with name maps?
+SYMS_API SYMS_USIDList syms_group_type_list_from_name_accelerated(SYMS_Arena *arena, SYMS_Group *group,
                                                                   SYMS_String8 name);
+SYMS_API SYMS_SymbolIDArray syms_group_types_from_unit_name(SYMS_Group *group, SYMS_UnitID uid, SYMS_String8 name);
 
-SYMS_API SYMS_SymbolIDArray syms_group_types_from_unit_name(SYMS_Group *group,
-                                                            SYMS_UnitID uid, SYMS_String8 name);
-
-SYMS_API SYMS_U64       syms_group_type_size_from_usid(SYMS_TypeGraph *graph,
-                                                       SYMS_Group *group, SYMS_USID usid);
+SYMS_API SYMS_U64       syms_group_type_size_from_usid(SYMS_TypeGraph *graph, SYMS_Group *group, SYMS_USID usid);
 
-SYMS_API SYMS_TypeMemberArray* syms_type_members_from_type(SYMS_Group *group,
-                                                           SYMS_TypeNode *type);
-SYMS_API SYMS_EnumInfoArray*   syms_type_enum_members_from_type(SYMS_Group *group,
-                                                                SYMS_TypeNode *type);
+SYMS_API SYMS_TypeMemberArray* syms_type_members_from_type(SYMS_Group *group, SYMS_TypeNode *type);
+SYMS_API SYMS_EnumInfoArray*   syms_type_enum_members_from_type(SYMS_Group *group, SYMS_TypeNode *type);
 
 ////////////////////////////////
-//~ NOTE(allen): Syms File Map
+//~ allen: Syms File Map
 
 SYMS_API SYMS_Name2FileIDMap* syms_group_file_map(SYMS_Group *group);
 
diff --git a/Engine/Source/Runtime/SymsLib/syms/core/pdb/syms_pdb_parser.c b/Engine/Source/Runtime/SymsLib/syms/core/pdb/syms_pdb_parser.c
--- a/Engine/Source/Runtime/SymsLib/syms/core/pdb/syms_pdb_parser.c
+++ b/Engine/Source/Runtime/SymsLib/syms/core/pdb/syms_pdb_parser.c
@@ -453,7 +453,7 @@ syms_pdb_symbols_from_name(SYMS_Arena *arena, SYMS_String8 data, SYMS_MsfAccel *
           {
             SYMS_CvRef2 ref2 = {0};
             syms_msf_read_struct_in_range(data, msf, element.range, 0, &ref2);
-            uid = ref2.imod + SYMS_PdbPseudoUnit_FIRST_COMP_UNIT;
+            uid = SYMS_PdbPseudoUnit_COUNT + ref2.imod;
             sid = SYMS_ID_u32_u32(SYMS_PdbSymbolIDKind_Off, ref2.sym_off);
           }break;
         }
@@ -1084,26 +1084,28 @@ syms_pdb_unit_info_from_uid(SYMS_PdbUnitSetAccel *unit_set, SYMS_UnitID uid){
   
   // TODO(allen): language
   
-  // try symbols pseudo unit
-  if (uid == SYMS_PdbPseudoUnit_SYM){
-    result.features = (SYMS_UnitFeature_StaticVariables|
-                       SYMS_UnitFeature_ExternVariables|
-                       SYMS_UnitFeature_FunctionStubs);
-  }
-  
-  // try types pseudo unit
-  else if (uid == SYMS_PdbPseudoUnit_TPI){
-    result.features = SYMS_UnitFeature_Types;
-  }
-  
-  // try compilation unit
-  else{
-    SYMS_PdbCompUnit *unit = syms_pdb_comp_unit_from_id(unit_set, uid);
-    if (unit != 0){
-      result.features = (SYMS_UnitFeature_CompilationUnit|
-                         SYMS_UnitFeature_StaticVariables|
-                         SYMS_UnitFeature_Functions);
-    }
+  switch (uid){
+    case SYMS_PdbPseudoUnit_SYM:
+    {
+      result.features = (SYMS_UnitFeature_StaticVariables|
+                         SYMS_UnitFeature_ExternVariables|
+                         SYMS_UnitFeature_FunctionStubs);
+    }break;
+    
+    case SYMS_PdbPseudoUnit_TPI:
+    {
+      result.features = SYMS_UnitFeature_Types;
+    }break;
+    
+    default:
+    {
+      SYMS_PdbCompUnit *unit = syms_pdb_comp_unit_from_id(unit_set, uid);
+      if (unit != 0){
+        result.features = (SYMS_UnitFeature_CompilationUnit|
+                           SYMS_UnitFeature_StaticVariables|
+                           SYMS_UnitFeature_Functions);
+      }
+    }break;
   }
   
   return(result);
@@ -3217,27 +3219,29 @@ syms_pdb_stub_from_unit_index(SYMS_PdbUnitAccel *unit, SYMS_U32 index){
 
 SYMS_API SYMS_PdbUnitAccel*
 syms_pdb_unit_accel_from_id(SYMS_Arena *arena, SYMS_String8 data, SYMS_PdbDbgAccel *dbg,
-                            SYMS_PdbUnitSetAccel *unit_set, SYMS_UnitID id){
+                            SYMS_PdbUnitSetAccel *unit_set, SYMS_UnitID uid){
   // setup result
   SYMS_PdbUnitAccel *result = (SYMS_PdbUnitAccel*)&syms_format_nil;
   
-  // try symbols pseudo unit
-  if (id == SYMS_PdbPseudoUnit_SYM){
-    result = syms_pdb_pub_sym_accel_from_dbg(arena, data, dbg);
-  }
-  
-  // try types pseudo unit
-  else if (id == SYMS_PdbPseudoUnit_TPI){
-    result = syms_pdb_leaf_accel_from_dbg(arena, data, dbg);
-  }
-  
-  // try compilation unit
-  else{
-    SYMS_PdbCompUnit *unit = syms_pdb_comp_unit_from_id(unit_set, id);
-    if (unit != 0){
-      SYMS_MsfRange range = syms_pdb_msf_range_from_comp_unit(unit, SYMS_PdbCompUnitRange_Symbols);
-      result = syms_pdb_sym_accel_from_range(arena, data, dbg->msf, range, id);
-    }
+  switch (uid){
+    case SYMS_PdbPseudoUnit_SYM:
+    {
+      result = syms_pdb_pub_sym_accel_from_dbg(arena, data, dbg);
+    }break;
+    
+    case SYMS_PdbPseudoUnit_TPI:
+    {
+      result = syms_pdb_leaf_accel_from_dbg(arena, data, dbg);
+    }break;
+    
+    default:
+    {
+      SYMS_PdbCompUnit *unit = syms_pdb_comp_unit_from_id(unit_set, uid);
+      if (unit != 0){
+        SYMS_MsfRange range = syms_pdb_msf_range_from_comp_unit(unit, SYMS_PdbCompUnitRange_Symbols);
+        result = syms_pdb_sym_accel_from_range(arena, data, dbg->msf, range, uid);
+      }
+    }break;
   }
   
   return(result);
@@ -3250,7 +3254,6 @@ syms_pdb_uid_from_accel(SYMS_PdbUnitAccel *unit){
 
 SYMS_API SYMS_UnitID
 syms_pdb_tls_var_uid_from_dbg(SYMS_PdbDbgAccel *dbg){
-  // TODO(allen): double check this API design
   // NOTE(nick): Thread var symbols are stored in global symbol stream
   // which doesn't abstract well with DWARF model where thread vars
   // are stored on per compilation units basis.
@@ -4780,86 +4783,6 @@ syms_pdb_scope_children_from_sid(SYMS_Arena *arena, SYMS_String8 data, SYMS_PdbD
   return(result);
 }
 
-SYMS_API void
-syms_pdb_stripped_sort_in_place(SYMS_StrippedInfo *info, SYMS_U64 count){
-  if (count > 1){
-    SYMS_U64 pivot = count - 1;
-    SYMS_U64 pivot_voff = info[pivot].voff;
-    
-    SYMS_B32 eq_flip_flop = syms_false;
-    
-    SYMS_U64 j = 0;
-    for (SYMS_U64 i = 0; i < pivot; i += 1){
-      SYMS_U64 voff = info[i].voff;
-      SYMS_B32 swap = syms_false;
-      if (voff < pivot_voff){
-        swap = syms_true;
-      }
-      else if (voff == pivot_voff){
-        swap = eq_flip_flop;
-        eq_flip_flop = !eq_flip_flop;
-      }
-      if (swap){
-        SYMS_Swap(SYMS_StrippedInfo, info[i], info[j]);
-        j += 1;
-      }
-    }
-    
-    SYMS_Swap(SYMS_StrippedInfo, info[j], info[pivot]);
-    
-    syms_pdb_stripped_sort_in_place(info, j);
-    syms_pdb_stripped_sort_in_place(info + j + 1, count - (j + 1));
-  }
-}
-
-SYMS_API SYMS_StrippedInfoArray
-syms_pdb_stripped_from_unit(SYMS_Arena *arena, SYMS_String8 data, SYMS_PdbDbgAccel *dbg, SYMS_PdbUnitAccel *unit){
-  SYMS_StrippedInfoArray result = {0};
-  
-  if (unit->pub_count != 0){
-    // setup accels
-    SYMS_MsfAccel *msf = dbg->msf;
-    
-    // allocate output memory
-    SYMS_U64 count = unit->pub_count;
-    SYMS_StrippedInfo *info = syms_push_array(arena, SYMS_StrippedInfo, count);
-    
-    // fill output array
-    SYMS_StrippedInfo *info_ptr = info;
-    SYMS_PdbStub **stub_opl = unit->pub_stubs + count;
-    for (SYMS_PdbStub **stub_ptr = unit->pub_stubs; stub_ptr < stub_opl; stub_ptr += 1, info_ptr += 1){
-      SYMS_PdbStub *stub = *stub_ptr;
-      
-      // extract info from stub
-      SYMS_String8 name = stub->name;
-      SYMS_U64 voff = 0;
-      {
-        SYMS_MsfRange range = syms_msf_range_from_sn(msf, unit->sn);
-        SYMS_CvElement sig_element = syms_cv_element(data, msf, range, stub->off);
-        SYMS_CvPubsym32 pubsym32 = {0};
-        if (syms_msf_read_struct_in_range(data, msf, sig_element.range, 0, &pubsym32)){
-          SYMS_CoffSection coff_section = syms_pdb_coff_section(data, dbg, pubsym32.sec);
-          voff = coff_section.virt_off + pubsym32.off;
-        }
-      }
-      
-      // fill stripped info
-      info_ptr->name = syms_push_string_copy(arena, name);
-      info_ptr->voff = voff;
-    }
-    
-    // sort stripped info
-    syms_pdb_stripped_sort_in_place(info, count);
-    
-    // fill result struct
-    result.info = info;
-    result.count = count;
-  }
-  
-  return(result);
-}
-
-
 ////////////////////////////////
 //~ NOTE(allen): PDB Signature Info
 
@@ -5449,7 +5372,7 @@ syms_pdb_type_map_from_dbg(SYMS_Arena *arena, SYMS_String8 data, SYMS_PdbDbgAcce
 }
 
 SYMS_API SYMS_PdbMapAccel*
-syms_pdb_image_symbol_map_from_dbg(SYMS_Arena *arena, SYMS_String8 data, SYMS_PdbDbgAccel *dbg){
+syms_pdb_unmangled_symbol_map_from_dbg(SYMS_Arena *arena, SYMS_String8 data, SYMS_PdbDbgAccel *dbg){
   SYMS_PdbMapAccel *result = syms_push_array(arena, SYMS_PdbMapAccel, 1);
   result->format = SYMS_FileFormat_PDB;
   result->uid = SYMS_PdbPseudoUnit_SYM;
@@ -5467,14 +5390,14 @@ syms_pdb_usid_list_from_string(SYMS_Arena *arena, SYMS_String8 data, SYMS_PdbDbg
   SYMS_USIDList result = {0};
   if (unit->format == SYMS_FileFormat_PDB){
     switch (map->uid){
-      case SYMS_PdbPseudoUnit_TPI:
+      case SYMS_PdbPseudoUnit_SYM:
       {
-        result = syms_pdb_types_from_name(arena, data, dbg, unit, string);
+        result = syms_pdb_symbols_from_name(arena, data, dbg->msf, &dbg->gsi, unit, string);
       }break;
       
-      case SYMS_PdbPseudoUnit_SYM:
+      case SYMS_PdbPseudoUnit_TPI:
       {
-        result = syms_pdb_symbols_from_name(arena, data, dbg->msf, &dbg->gsi, unit, string);
+        result = syms_pdb_types_from_name(arena, data, dbg, unit, string);
       }break;
     }
   }
@@ -5483,7 +5406,112 @@ syms_pdb_usid_list_from_string(SYMS_Arena *arena, SYMS_String8 data, SYMS_PdbDbg
 
 
 ////////////////////////////////
-// NOTE(allen): PDB CV -> Syms Enums and Flags
+//~ allen: PDB Mangled Names
+
+SYMS_API SYMS_UnitID
+syms_pdb_link_names_uid(void){
+  return(SYMS_PdbPseudoUnit_SYM);
+}
+
+SYMS_API SYMS_PdbLinkMapAccel*
+syms_pdb_link_map_from_dbg(SYMS_Arena *arena, SYMS_String8 data, SYMS_PdbDbgAccel *dbg){
+  SYMS_PdbLinkMapAccel *result = syms_push_array(arena, SYMS_PdbLinkMapAccel, 1);
+  result->format = SYMS_FileFormat_PDB;
+  return(result);
+}
+
+SYMS_API SYMS_U64
+syms_pdb_voff_from_link_name(SYMS_String8 data, SYMS_PdbDbgAccel *dbg, SYMS_PdbLinkMapAccel *map,
+                             SYMS_PdbUnitAccel *link_unit, SYMS_String8 name){
+  SYMS_U64 result = 0;
+  
+  // get accelerators
+  SYMS_MsfAccel *msf = dbg->msf;
+  SYMS_PdbGsiAccel *psi = &dbg->psi;
+  
+  if (psi->bucket_count > 0){
+    // get unit range
+    SYMS_MsfRange range = syms_msf_range_from_sn(msf, link_unit->sn);
+    
+    // grab scratch
+    SYMS_ArenaTemp scratch = syms_get_scratch(0, 0);
+    
+    // get bucket
+    SYMS_U32 name_hash = syms_pdb_hashV1(name);
+    SYMS_U32 bucket_index = name_hash%psi->bucket_count;
+    
+    // iterate bucket
+    for (SYMS_PdbChain *bucket = psi->buckets[bucket_index];
+         bucket != 0;
+         bucket = bucket->next){
+      SYMS_U32 off = bucket->v;
+      SYMS_PdbStub *stub = syms_pdb_stub_from_unit_off(link_unit, off);
+      if (stub != 0 && syms_string_match(stub->name, name, 0)){
+        SYMS_CvElement element = syms_cv_element(data, msf, range, stub->off);
+        if (element.kind == SYMS_CvSymKind_PUB32){
+          SYMS_CvPubsym32 pub = {0};
+          syms_msf_read_struct_in_range(data, msf, element.range, 0, &pub);
+          SYMS_CoffSection coff_section = syms_pdb_coff_section(data, dbg, pub.sec);
+          result = coff_section.virt_off + pub.off;
+          break;
+        }
+      }
+    }
+    
+    // release scratch
+    syms_release_scratch(scratch);
+  }
+  
+  return(result);
+}
+
+SYMS_API SYMS_LinkNameRecArray
+syms_pdb_link_name_array_from_unit(SYMS_Arena *arena, SYMS_String8 data,
+                                   SYMS_PdbDbgAccel *dbg, SYMS_PdbUnitAccel *unit){
+  SYMS_LinkNameRecArray result = {0};
+  
+  if (unit->pub_count != 0){
+    // setup accels
+    SYMS_MsfAccel *msf = dbg->msf;
+    
+    // allocate output memory
+    SYMS_U64 count = unit->pub_count;
+    SYMS_LinkNameRec *rec = syms_push_array(arena, SYMS_LinkNameRec, count);
+    
+    // fill output array
+    SYMS_LinkNameRec *rec_ptr = rec;
+    SYMS_PdbStub **stub_opl = unit->pub_stubs + count;
+    for (SYMS_PdbStub **stub_ptr = unit->pub_stubs; stub_ptr < stub_opl; stub_ptr += 1, rec_ptr += 1){
+      SYMS_PdbStub *stub = *stub_ptr;
+      
+      // extract rec from stub
+      SYMS_String8 name = stub->name;
+      SYMS_U64 voff = 0;
+      {
+        SYMS_MsfRange range = syms_msf_range_from_sn(msf, unit->sn);
+        SYMS_CvElement sig_element = syms_cv_element(data, msf, range, stub->off);
+        SYMS_CvPubsym32 pubsym32 = {0};
+        if (syms_msf_read_struct_in_range(data, msf, sig_element.range, 0, &pubsym32)){
+          SYMS_CoffSection coff_section = syms_pdb_coff_section(data, dbg, pubsym32.sec);
+          voff = coff_section.virt_off + pubsym32.off;
+        }
+      }
+      
+      // fill stripped rec
+      rec_ptr->name = syms_push_string_copy(arena, name);
+      rec_ptr->voff = voff;
+    }
+    
+    // fill result struct
+    result.recs = rec;
+    result.count = count;
+  }
+  
+  return(result);
+}
+
+////////////////////////////////
+//~ allen: PDB CV -> Syms Enums and Flags
 
 SYMS_API SYMS_TypeKind
 syms_pdb_type_kind_from_cv_pointer_mode(SYMS_CvPointerMode mode){
diff --git a/Engine/Source/Runtime/SymsLib/syms/core/pdb/syms_pdb_parser.h b/Engine/Source/Runtime/SymsLib/syms/core/pdb/syms_pdb_parser.h
--- a/Engine/Source/Runtime/SymsLib/syms/core/pdb/syms_pdb_parser.h
+++ b/Engine/Source/Runtime/SymsLib/syms/core/pdb/syms_pdb_parser.h
@@ -5,12 +5,12 @@
 #define SYMS_PDB_PARSER_H
 
 ////////////////////////////////
-//~ NOTE(allen): PDB's Prerequisite Formats
+//~ allen: PDB's Prerequisite Formats
 
 #include "syms/core/pdb/syms_msf_parser.h"
 
 ////////////////////////////////
-//~ NOTE(allen): PDB Fundamental
+//~ allen: PDB Fundamental
 
 // TODO(allen): version of this with optimized memory locality could be pretty useful.
 typedef struct SYMS_PdbChain{
@@ -40,7 +40,7 @@ typedef enum SYMS_PdbPseudoUnit{
 } SYMS_PdbPseudoUnit;
 
 ////////////////////////////////
-//~ NOTE(allen): CV Parsing
+//~ allen: CV Parsing
 
 typedef struct SYMS_CvElement{
   SYMS_MsfRange range;
@@ -54,7 +54,7 @@ typedef struct SYMS_CvNumeric{
 } SYMS_CvNumeric;
 
 ////////////////////////////////
-//~ NOTE(allen): PDB DBI Information
+//~ allen: PDB DBI Information
 
 typedef struct SYMS_PdbDbiAccel{
   SYMS_B16 valid;
@@ -91,7 +91,7 @@ typedef struct SYMS_PdbStrtblAccel{
 } SYMS_PdbStrtblAccel;
 
 ////////////////////////////////
-//~ NOTE(allen): PDB Units
+//~ allen: PDB Units
 
 typedef struct SYMS_PdbCompUnit{
   SYMS_MsfStreamNumber sn;
@@ -170,7 +170,7 @@ typedef struct SYMS_PdbLeafResolve{
 } SYMS_PdbLeafResolve;
 
 ////////////////////////////////
-//~ NOTE(allen): PDB C13
+//~ allen: PDB C13
 
 typedef struct SYMS_PdbC13SubSection{
   struct SYMS_PdbC13SubSection *next;
@@ -180,7 +180,7 @@ typedef struct SYMS_PdbC13SubSection{
 } SYMS_PdbC13SubSection;
 
 ////////////////////////////////
-//~ NOTE(allen): PDB TPI
+//~ allen: PDB TPI
 
 typedef struct SYMS_PdbTpiOffRange{
   SYMS_U32 first_off;
@@ -201,7 +201,7 @@ typedef struct SYMS_PdbTpiAccel{
 } SYMS_PdbTpiAccel;
 
 ////////////////////////////////
-//~ NOTE(allen): PDB GSI
+//~ allen: PDB GSI
 
 typedef struct SYMS_PdbGsiAccel{
   SYMS_PdbChain **buckets;
@@ -209,7 +209,7 @@ typedef struct SYMS_PdbGsiAccel{
 } SYMS_PdbGsiAccel;
 
 ////////////////////////////////
-//~ NOTE(allen): PDB Accel Types
+//~ allen: PDB Accel Types
 
 typedef struct SYMS_PdbFileAccel{
   SYMS_FileFormat format;
@@ -230,7 +230,7 @@ typedef struct SYMS_PdbDbgAccel{
 } SYMS_PdbDbgAccel;
 
 ////////////////////////////////
-//~ NOTE(allen): PDB Members
+//~ allen: PDB Members
 
 typedef struct SYMS_PdbMemStubNode{
   struct SYMS_PdbMemStubNode *next;
@@ -259,7 +259,7 @@ typedef struct SYMS_EnumInfoNode{
 } SYMS_EnumInfoNode;
 
 ////////////////////////////////
-//~ NOTE(allen): PDB Lines
+//~ allen: PDB Lines
 
 typedef struct SYMS_PdbFileNode{
   struct SYMS_PdbFileNode *next;
@@ -286,15 +286,19 @@ typedef struct SYMS_PdbLineTableLoose{
 
 
 ////////////////////////////////
-//~ NOTE(allen): PDB Map Accel Type
+//~ allen: PDB Map Accel Type
 
 typedef struct SYMS_PdbMapAccel{
   SYMS_FileFormat format;
   SYMS_UnitID uid;
 } SYMS_PdbMapAccel;
 
+typedef struct SYMS_PdbLinkMapAccel{
+  SYMS_FileFormat format;
+} SYMS_PdbLinkMapAccel;
+
 ////////////////////////////////
-//~ NOTE(allen): PDB TPI Functions
+//~ allen: PDB TPI Functions
 
 //- tpi accel helpers
 SYMS_API SYMS_PdbTpiOffRange syms_pdb_tpi__hint_from_index(SYMS_PdbTpiAccel *tpi, SYMS_CvTypeIndex ti);
@@ -315,7 +319,7 @@ SYMS_API SYMS_USIDList syms_pdb_types_from_name(SYMS_Arena *arena, SYMS_String8
 
 
 ////////////////////////////////
-//~ NOTE(allen): PDB GSI Functions
+//~ allen: PDB GSI Functions
 
 SYMS_API SYMS_PdbGsiAccel  syms_pdb_gsi_accel_from_range(SYMS_Arena *arena, SYMS_String8 data, SYMS_MsfAccel *msf,
                                                          SYMS_MsfRange range);
@@ -327,7 +331,7 @@ SYMS_API SYMS_USIDList     syms_pdb_symbols_from_name(SYMS_Arena *arena, SYMS_St
 
 
 ////////////////////////////////
-//~ NOTE(allen): PDB Accelerator Setup
+//~ allen: PDB Accelerator Setup
 
 // pdb specific parsing
 SYMS_API SYMS_PdbDbiAccel syms_pdb_dbi_from_msf(SYMS_String8 data, SYMS_MsfAccel *msf);
@@ -354,20 +358,20 @@ SYMS_API SYMS_PdbDbgAccel*  syms_pdb_dbg_accel_from_file(SYMS_Arena *arena, SYMS
 
 
 ////////////////////////////////
-//~ NOTE(allen): PDB Architecture
+//~ allen: PDB Architecture
 
 // main api
 SYMS_API SYMS_Arch syms_pdb_arch_from_dbg(SYMS_PdbDbgAccel *dbg);
 
 
 ////////////////////////////////
-//~ NOTE(allen): PDB Match Keys
+//~ allen: PDB Match Keys
 
 SYMS_API SYMS_ExtMatchKey syms_pdb_ext_match_key_from_dbg(SYMS_String8 data, SYMS_PdbDbgAccel *dbg);
 
 
 ////////////////////////////////
-//~ NOTE(allen): PDB Sections
+//~ allen: PDB Sections
 
 // pdb specific
 SYMS_API SYMS_CoffSection syms_pdb_coff_section(SYMS_String8 data, SYMS_PdbDbgAccel *dbg, SYMS_U64 n);
@@ -377,7 +381,7 @@ SYMS_API SYMS_SecInfoArray syms_pdb_sec_info_array_from_dbg(SYMS_Arena *arena, S
                                                             SYMS_PdbDbgAccel *dbg);
 
 ////////////////////////////////
-//~ NOTE(allen): PDB Compilation Units
+//~ allen: PDB Compilation Units
 
 // pdb specific
 SYMS_API SYMS_PdbCompUnit* syms_pdb_comp_unit_from_id(SYMS_PdbUnitSetAccel *unit_set, SYMS_UnitID id);
@@ -398,7 +402,7 @@ SYMS_API SYMS_UnitRangeArray  syms_pdb_unit_ranges_from_set(SYMS_Arena *arena,
 SYMS_API SYMS_UnitFeatures    syms_pdb_unit_features_from_number(SYMS_PdbUnitSetAccel *unit_set, SYMS_U64 n);
 
 ////////////////////////////////
-//~ NOTE(allen): PDB Symbol Parsing
+//~ allen: PDB Symbol Parsing
 
 // cv parsing helpers
 SYMS_API void               syms_pdb_leaf_debug_helper(SYMS_CvLeaf kind);
@@ -475,7 +479,7 @@ SYMS_API SYMS_ConstInfo    syms_pdb_const_info_from_id(SYMS_String8 data, SYMS_P
                                                        SYMS_PdbUnitAccel *unit, SYMS_SymbolID id);
 
 ////////////////////////////////
-//~ NOTE(allen): PDB Variable Info
+//~ allen: PDB Variable Info
 
 // cv parse
 SYMS_API SYMS_USID   syms_pdb_sym_type_from_var_id(SYMS_String8 data, SYMS_PdbDbgAccel *dbg,
@@ -507,7 +511,7 @@ SYMS_API SYMS_Location syms_pdb_location_from_id(SYMS_Arena *arena, SYMS_String8
 
 
 ////////////////////////////////
-//~ NOTE(allen): PDB Member Info
+//~ allen: PDB Member Info
 
 // cv parse
 SYMS_API void syms_pdb__field_list_parse(SYMS_Arena *arena, SYMS_String8 data, SYMS_MsfAccel *msf,
@@ -535,7 +539,7 @@ SYMS_API SYMS_EnumInfoArray syms_pdb_enum_info_array_from_sid(SYMS_Arena *arena,
                                                               SYMS_SymbolID sid);
 
 ////////////////////////////////
-//~ NOTE(allen): PDB Procedure Info
+//~ allen: PDB Procedure Info
 
 // main api
 
@@ -552,13 +556,8 @@ SYMS_API SYMS_SymbolIDArray syms_pdb_scope_children_from_sid(SYMS_Arena *arena,
                                                              SYMS_PdbDbgAccel *dbg, SYMS_PdbUnitAccel *unit,
                                                              SYMS_SymbolID id);
 
-SYMS_API void                   syms_pdb_stripped_sort_in_place(SYMS_StrippedInfo *info, SYMS_U64 count);
-
-SYMS_API SYMS_StrippedInfoArray syms_pdb_stripped_from_unit(SYMS_Arena *arena, SYMS_String8 data,
-                                                            SYMS_PdbDbgAccel *dbg, SYMS_PdbUnitAccel *unit);
-
 ////////////////////////////////
-//~ NOTE(allen): PDB Signature Info
+//~ allen: PDB Signature Info
 
 // pdb specific helper
 SYMS_API SYMS_SigInfo syms_pdb_sig_info_from_sig_index(SYMS_Arena *arena, SYMS_String8 data, SYMS_PdbDbgAccel *dbg,
@@ -569,7 +568,7 @@ SYMS_API SYMS_SigInfo syms_pdb_sig_info_from_id(SYMS_Arena *arena, SYMS_String8
                                                 SYMS_PdbUnitAccel *unit, SYMS_SymbolID id);
 
 ////////////////////////////////
-//~ NOTE(allen): PDB Line Info
+//~ allen: PDB Line Info
 
 // pdb specific
 SYMS_API SYMS_PdbC13SubSection* syms_pdb_c13_sub_sections(SYMS_Arena *arena, SYMS_String8 data, 
@@ -589,19 +588,33 @@ SYMS_API SYMS_LineParseOut syms_pdb_line_parse_from_uid(SYMS_Arena *arena, SYMS_
                                                         SYMS_UnitID uid);
 
 ////////////////////////////////
-//~ NOTE(allen): PDB Name Maps
+//~ allen: PDB Name Maps
 
 // main api
 SYMS_API SYMS_PdbMapAccel* syms_pdb_type_map_from_dbg(SYMS_Arena *arena, SYMS_String8 data, SYMS_PdbDbgAccel *dbg);
-SYMS_API SYMS_PdbMapAccel* syms_pdb_image_symbol_map_from_dbg(SYMS_Arena *arena, SYMS_String8 data,
-                                                              SYMS_PdbDbgAccel *dbg);
+SYMS_API SYMS_PdbMapAccel* syms_pdb_unmangled_symbol_map_from_dbg(SYMS_Arena *arena, SYMS_String8 data,
+                                                                  SYMS_PdbDbgAccel *dbg);
 SYMS_API SYMS_UnitID       syms_pdb_partner_uid_from_map(SYMS_PdbMapAccel *map);
 SYMS_API SYMS_USIDList     syms_pdb_usid_list_from_string(SYMS_Arena *arena, SYMS_String8 data,
                                                           SYMS_PdbDbgAccel *dbg, SYMS_PdbUnitAccel *unit,
                                                           SYMS_PdbMapAccel *map, SYMS_String8 string);
 
 ////////////////////////////////
-//~ NOTE(allen): PDB CV -> Syms Enums and Flags
+//~ allen: PDB Mangled Names
+
+// main api
+SYMS_API SYMS_UnitID syms_pdb_link_names_uid(void);
+
+SYMS_API SYMS_PdbLinkMapAccel* syms_pdb_link_map_from_dbg(SYMS_Arena *arena, SYMS_String8 data,
+                                                          SYMS_PdbDbgAccel *dbg);
+SYMS_API SYMS_U64              syms_pdb_voff_from_link_name(SYMS_String8 data, SYMS_PdbDbgAccel *dbg,
+                                                            SYMS_PdbLinkMapAccel *map,
+                                                            SYMS_PdbUnitAccel *link_unit, SYMS_String8 name);
+SYMS_API SYMS_LinkNameRecArray syms_pdb_link_name_array_from_unit(SYMS_Arena *arena, SYMS_String8 data,
+                                                                  SYMS_PdbDbgAccel *dbg, SYMS_PdbUnitAccel *unit);
+
+////////////////////////////////
+//~ allen: PDB CV -> Syms Enums and Flags
 
 SYMS_API SYMS_TypeKind       syms_pdb_type_kind_from_cv_pointer_mode(SYMS_CvPointerMode mode);
 SYMS_API SYMS_TypeModifiers  syms_pdb_modifier_from_cv_pointer_attribs(SYMS_CvPointerAttribs attribs);
diff --git a/Engine/Source/Runtime/SymsLib/syms/core/pe/syms_pe.h b/Engine/Source/Runtime/SymsLib/syms/core/pe/syms_pe.h
--- a/Engine/Source/Runtime/SymsLib/syms/core/pe/syms_pe.h
+++ b/Engine/Source/Runtime/SymsLib/syms/core/pe/syms_pe.h
@@ -30,6 +30,69 @@
 // this is the "MZ" as a 16-bit short
 #define SYMS_DOS_MAGIC 0x5a4d
 
+////////////////////////////////
+
+#define SYMS_COFF_MIN_BIG_OBJ_VERSION 2
+#define SYMS_COFF_BIG_OBJ_MAGIC_LO 0xC7A1BAD1EEBAA94Bull
+#define SYMS_COFF_BIG_OBJ_MAGIC_HI 0xAF20FAF66AA4DCB8ull
+typedef struct SYMS_CoffHeaderBigObj{
+  SYMS_U16 sig1; // SYMS_CoffMachineType_UNKNOWN
+  SYMS_U16 sig2; // SYMS_U16_MAX
+  SYMS_U16 version;
+  SYMS_U16 machine;
+  SYMS_U32 time_stamp;
+  SYMS_U8 magic[16];
+  SYMS_U32 unused[4];
+  SYMS_U32 section_count;
+  SYMS_U32 pointer_to_symbol_table;
+  SYMS_U32 number_of_symbols;
+} SYMS_CoffHeaderBigObj;
+
+#pragma pack(push, 1)
+
+typedef union SYMS_CoffSymbolName{
+  SYMS_U8 short_name[8];
+  struct{
+    // if this field is filled with zeroes we have a long name,
+    // which means name is stored in the string table 
+    // and we need to use the offset to look it up...
+    SYMS_U32 zeroes;
+    SYMS_U32 string_table_offset;
+  } long_name;
+} SYMS_CoffSymbolName;
+
+typedef struct SYMS_CoffSymbol16{
+  SYMS_CoffSymbolName name;
+  SYMS_U32 value;
+  SYMS_U16 section_number;
+  union{
+    struct{
+      SYMS_CoffSymDType msb;
+      SYMS_CoffSymType  lsb;
+    } u;
+    SYMS_U16 v;
+  } type;
+  SYMS_CoffSymStorageClass storage_class;
+  SYMS_U8 number_of_aux_symbols;
+} SYMS_CoffSymbol16;
+
+typedef struct SYMS_CoffSymbol32{
+  SYMS_CoffSymbolName name;
+  SYMS_U32 value;
+  SYMS_U32 section_number;
+  union{
+    struct{
+      SYMS_CoffSymDType msb;
+      SYMS_CoffSymType  lsb;
+    } u;
+    SYMS_U16 v;
+  } type;
+  SYMS_CoffSymStorageClass storage_class;
+  SYMS_U8 number_of_aux_symbols;
+} SYMS_CoffSymbol32;
+
+#pragma pack(pop)
+
 ////////////////////////////////
 // NOTE(allen): PE Magic Numbers
 
@@ -100,8 +163,7 @@ typedef struct SYMS_PeCvHeaderPDB70{
 ////////////////////////////////
 // Import
 
-typedef struct
-{
+typedef struct SYMS_PeImportDirectoryEntry{
   SYMS_U32 lookup_table_virt_off;
   SYMS_U32 timestamp;
   SYMS_U32 forwarder_chain;
@@ -112,8 +174,7 @@ typedef struct
 ////////////////////////////////
 // Export
 
-typedef struct
-{
+typedef struct SYMS_PeExportTable{
   SYMS_U32 flags;                       // must be zero
   SYMS_U32 timestamp;                   // time and date when export table was created
   SYMS_U16 major_ver;                   // table version, user can change major and minor version
diff --git a/Engine/Source/Runtime/SymsLib/syms/core/pe/syms_pe_parser.c b/Engine/Source/Runtime/SymsLib/syms/core/pe/syms_pe_parser.c
--- a/Engine/Source/Runtime/SymsLib/syms/core/pe/syms_pe_parser.c
+++ b/Engine/Source/Runtime/SymsLib/syms/core/pe/syms_pe_parser.c
@@ -80,9 +80,9 @@
         {
           SYMS_PeOptionalPe32 pe_optional = {0};
           syms_based_range_read_struct(base, optional_range, 0, &pe_optional);
-
+          
           image_base = pe_optional.image_base;
-
+          
           reported_data_dir_offset = sizeof(pe_optional);
           reported_data_dir_count = pe_optional.data_dir_count;
         }break;
@@ -90,9 +90,9 @@
         {
           SYMS_PeOptionalPe32Plus pe_optional = {0};
           syms_based_range_read_struct(base, optional_range, 0, &pe_optional);
-
+          
           image_base = pe_optional.image_base;
-
+          
           reported_data_dir_offset = sizeof(pe_optional);
           reported_data_dir_count = pe_optional.data_dir_count;
         }break;
@@ -176,12 +176,32 @@
   return(result);
 }
 
+SYMS_API SYMS_CoffFileAccel*
+syms_coff_file_accel_from_data(SYMS_Arena *arena, SYMS_String8 data){
+  SYMS_CoffFileAccel *result = (SYMS_CoffFileAccel *)&syms_format_nil;
+  // TODO(allen): 
+  return(result);
+}
+
+SYMS_API SYMS_CoffBinAccel*
+syms_coff_bin_accel_from_file(SYMS_Arena *arena, SYMS_String8 data, SYMS_CoffFileAccel *file){
+  SYMS_CoffBinAccel *result = (SYMS_CoffBinAccel *)&syms_format_nil;
+  // TODO(allen): 
+  return(result);
+}
+
 SYMS_API SYMS_Arch
 syms_pe_arch_from_bin(SYMS_PeBinAccel *bin){
   SYMS_Arch result = bin->arch;
   return(result);
 }
 
+SYMS_API SYMS_Arch
+syms_coff_arch_from_bin(SYMS_CoffBinAccel *bin){
+  SYMS_Arch result = bin->arch;
+  return(result);
+}
+
 SYMS_API SYMS_ExtFileList
 syms_pe_ext_file_list_from_bin(SYMS_Arena *arena, SYMS_String8 data, SYMS_PeBinAccel *bin){
   SYMS_ExtFileList result = {0};
@@ -209,5 +229,13 @@
 }
 
 SYMS_API SYMS_SecInfoArray
+syms_pe_coff_sec_info_array_from_data(SYMS_Arena *arena, SYMS_String8 data,
+                                      SYMS_U64 sec_array_off, SYMS_U64 sec_count){
+  SYMS_SecInfoArray result = {0};
+  // TODO(allen): 
+  return(result);
+}
+
+SYMS_API SYMS_SecInfoArray
 syms_pe_sec_info_array_from_bin(SYMS_Arena *arena, SYMS_String8 data, SYMS_PeBinAccel *bin){
   SYMS_SecInfoArray result = {0};
@@ -235,142 +263,16 @@
   return(result);
 }
 
-////////////////////////////////
-// NOTE(allen): PE Specific Helpers
-
-SYMS_API SYMS_U64
-syms_pe_binary_search_intel_pdata(SYMS_String8 data, SYMS_PeBinAccel *bin, SYMS_U64 voff){
-  SYMS_ASSERT(bin->arch == SYMS_Arch_X86 || bin->arch == SYMS_Arch_X64);
-  SYMS_U64 pdata_off = bin->data_dirs_file[SYMS_PeDataDirectoryIndex_EXCEPTIONS].min;
-  SYMS_U64 pdata_count = syms_u64_range_size(bin->data_dirs_file[SYMS_PeDataDirectoryIndex_EXCEPTIONS]) / sizeof(SYMS_PeIntelPdata);
-  
-  SYMS_U64 result = 0;
-  // check if this bin includes a pdata array
-  if (pdata_count > 0){
-    SYMS_PeIntelPdata *pdata_array = (SYMS_PeIntelPdata*)(data.str + pdata_off);
-    if (voff >= pdata_array[0].voff_first){
-      
-      // binary search:
-      //  find max index s.t. pdata_array[index].voff_first <= voff
-      //  we assume (i < j) -> (pdata_array[i].voff_first < pdata_array[j].voff_first)
-      SYMS_U64 index = pdata_count;
-      SYMS_U64 min = 0;
-      SYMS_U64 opl = pdata_count;
-      for (;;){
-        SYMS_U64 mid = (min + opl)/2;
-        SYMS_PeIntelPdata *pdata = pdata_array + mid;
-        if (voff < pdata->voff_first){
-          opl = mid;
-        }
-        else if (pdata->voff_first < voff){
-          min = mid;
-        }
-        else{
-          index = mid;
-          break;
-        }
-        if (min + 1 >= opl){
-          index = min;
-          break;
-        }
-      }
-      
-      // if we are in range fill result
-      {
-        SYMS_PeIntelPdata *pdata = pdata_array + index;
-        if (pdata->voff_first <= voff && voff < pdata->voff_one_past_last){
-          result = pdata_off + index*sizeof(SYMS_PeIntelPdata);
-        }
-      }
-    }
-  }
-  
+SYMS_API SYMS_SecInfoArray
+syms_coff_sec_info_array_from_bin(SYMS_Arena *arena, SYMS_String8 data,
+                                  SYMS_CoffBinAccel *bin){
+  SYMS_SecInfoArray result = {0};
+  // TODO(allen): 
   return(result);
 }
 
-SYMS_API SYMS_U64
-syms_pe_sec_number_from_voff(SYMS_String8 data, SYMS_PeBinAccel *bin, SYMS_U64 voff){
-  SYMS_U64 sec_count = bin->section_count;
-  SYMS_CoffSection *sec_array = (SYMS_CoffSection*)((SYMS_U8*)data.str + bin->section_array_off);
-  SYMS_CoffSection *sec_ptr = sec_array;
-  SYMS_U64 result = 0;
-  for (SYMS_U64 i = 1; i <= sec_count; i += 1, sec_ptr += 1){
-    if (sec_ptr->virt_off <= voff && voff < sec_ptr->virt_off + sec_ptr->virt_size){
-      result = i;
-      break;
-    }
-  }
-  return(result);
-}
-
-SYMS_API void*
-syms_pe_ptr_from_sec_number(SYMS_String8 data, SYMS_PeBinAccel *bin, SYMS_U64 n){
-  void *result = 0;
-  SYMS_U64 sec_count = bin->section_count;
-  if (1 <= n && n <= sec_count){
-    SYMS_CoffSection *sec_array = (SYMS_CoffSection*)((SYMS_U8*)data.str + bin->section_array_off);
-    SYMS_CoffSection *sec = sec_array + n - 1;
-    if (sec->file_size > 0){
-      result = data.str + sec->file_off;
-    }
-  }
-  return(result);
-}
-
-SYMS_API void*
-syms_pe_ptr_from_foff(SYMS_String8 data, SYMS_PeBinAccel *bin, SYMS_U64 foff){
-  void *result = 0;
-  if (foff < data.size){
-    result = data.str + foff;
-  }
-  return(result);
-}
-
-SYMS_API void*
-syms_pe_ptr_from_voff(SYMS_String8 data, SYMS_PeBinAccel *bin, SYMS_U64 voff){
-  //- get the section for this voff
-  SYMS_U64 sec_count = bin->section_count;
-  SYMS_CoffSection *sec_array = (SYMS_CoffSection*)((SYMS_U8*)data.str + bin->section_array_off);
-  SYMS_CoffSection *sec_ptr = sec_array;
-  SYMS_CoffSection *sec = 0;
-  for (SYMS_U64 i = 1; i <= sec_count; i += 1, sec_ptr += 1){
-    if (sec_ptr->virt_off <= voff && voff < sec_ptr->virt_off + sec_ptr->virt_size){
-      sec = sec_ptr;
-      break;
-    }
-  }
-  
-  //- adjust to file pointer
-  void *result = 0;
-  if (sec != 0 & sec_ptr->file_size > 0){
-    result = data.str + voff - sec->virt_off + sec->file_off;
-  }
-  
-  return(result);
-}
-
-SYMS_API SYMS_U64
-syms_pe_virt_off_to_file_off(SYMS_CoffSection *sections, SYMS_U32 section_count, SYMS_U64 virt_off)
-{
-  SYMS_U64 file_off = 0;
-  for (SYMS_U32 sect_idx = 0; sect_idx < section_count; sect_idx += 1){
-    SYMS_CoffSection *sect = &sections[sect_idx];
-    if (sect->virt_off <= virt_off && virt_off < sect->virt_off + sect->virt_size){
-      file_off = sect->file_off + (virt_off - sect->virt_off);
-    }
-  }
-  return file_off;
-}
-
-SYMS_API SYMS_U64
-syms_pe_bin_virt_off_to_file_off(SYMS_String8 data, SYMS_PeBinAccel *bin, SYMS_U64 virt_off)
-{
-  return syms_pe_virt_off_to_file_off((SYMS_CoffSection*)((SYMS_U8*)data.str + bin->section_array_off), bin->section_count, virt_off);
-}
-
 SYMS_API SYMS_ImportArray
-syms_pe_imports_from_bin(SYMS_Arena *arena, SYMS_String8 data, SYMS_PeBinAccel *bin)
-{
+syms_pe_imports_from_bin(SYMS_Arena *arena, SYMS_String8 data, SYMS_PeBinAccel *bin){
   SYMS_ArenaTemp scratch = syms_get_scratch(&arena, 1);
   
   //- rjf: grab prerequisites
@@ -492,8 +394,7 @@
 }
 
 SYMS_API SYMS_ExportArray
-syms_pe_exports_from_bin(SYMS_Arena *arena, SYMS_String8 data, SYMS_PeBinAccel *bin)
-{
+syms_pe_exports_from_bin(SYMS_Arena *arena, SYMS_String8 data, SYMS_PeBinAccel *bin){
   SYMS_U64Range dir_range = bin->data_dirs_file[SYMS_PeDataDirectoryIndex_EXPORT];
   SYMS_U64Range data_range = syms_make_u64_range(0, data.size);
   void *base = (void*)data.str;
@@ -566,4 +467,136 @@
   return export_array;
 }
 
+
+////////////////////////////////
+// NOTE(allen): PE Specific Helpers
+
+SYMS_API SYMS_U64
+syms_pe_binary_search_intel_pdata(SYMS_String8 data, SYMS_PeBinAccel *bin, SYMS_U64 voff){
+  SYMS_ASSERT(bin->arch == SYMS_Arch_X86 || bin->arch == SYMS_Arch_X64);
+  SYMS_U64 pdata_off = bin->data_dirs_file[SYMS_PeDataDirectoryIndex_EXCEPTIONS].min;
+  SYMS_U64 pdata_count = syms_u64_range_size(bin->data_dirs_file[SYMS_PeDataDirectoryIndex_EXCEPTIONS]) / sizeof(SYMS_PeIntelPdata);
+  
+  SYMS_U64 result = 0;
+  // check if this bin includes a pdata array
+  if (pdata_count > 0){
+    SYMS_PeIntelPdata *pdata_array = (SYMS_PeIntelPdata*)(data.str + pdata_off);
+    if (voff >= pdata_array[0].voff_first){
+      
+      // binary search:
+      //  find max index s.t. pdata_array[index].voff_first <= voff
+      //  we assume (i < j) -> (pdata_array[i].voff_first < pdata_array[j].voff_first)
+      SYMS_U64 index = pdata_count;
+      SYMS_U64 min = 0;
+      SYMS_U64 opl = pdata_count;
+      for (;;){
+        SYMS_U64 mid = (min + opl)/2;
+        SYMS_PeIntelPdata *pdata = pdata_array + mid;
+        if (voff < pdata->voff_first){
+          opl = mid;
+        }
+        else if (pdata->voff_first < voff){
+          min = mid;
+        }
+        else{
+          index = mid;
+          break;
+        }
+        if (min + 1 >= opl){
+          index = min;
+          break;
+        }
+      }
+      
+      // if we are in range fill result
+      {
+        SYMS_PeIntelPdata *pdata = pdata_array + index;
+        if (pdata->voff_first <= voff && voff < pdata->voff_one_past_last){
+          result = pdata_off + index*sizeof(SYMS_PeIntelPdata);
+        }
+      }
+    }
+  }
+  
+  return(result);
+}
+
+SYMS_API SYMS_U64
+syms_pe_sec_number_from_voff(SYMS_String8 data, SYMS_PeBinAccel *bin, SYMS_U64 voff){
+  SYMS_U64 sec_count = bin->section_count;
+  SYMS_CoffSection *sec_array = (SYMS_CoffSection*)((SYMS_U8*)data.str + bin->section_array_off);
+  SYMS_CoffSection *sec_ptr = sec_array;
+  SYMS_U64 result = 0;
+  for (SYMS_U64 i = 1; i <= sec_count; i += 1, sec_ptr += 1){
+    if (sec_ptr->virt_off <= voff && voff < sec_ptr->virt_off + sec_ptr->virt_size){
+      result = i;
+      break;
+    }
+  }
+  return(result);
+}
+
+SYMS_API void*
+syms_pe_ptr_from_sec_number(SYMS_String8 data, SYMS_PeBinAccel *bin, SYMS_U64 n){
+  void *result = 0;
+  SYMS_U64 sec_count = bin->section_count;
+  if (1 <= n && n <= sec_count){
+    SYMS_CoffSection *sec_array = (SYMS_CoffSection*)((SYMS_U8*)data.str + bin->section_array_off);
+    SYMS_CoffSection *sec = sec_array + n - 1;
+    if (sec->file_size > 0){
+      result = data.str + sec->file_off;
+    }
+  }
+  return(result);
+}
+
+SYMS_API void*
+syms_pe_ptr_from_foff(SYMS_String8 data, SYMS_PeBinAccel *bin, SYMS_U64 foff){
+  void *result = 0;
+  if (foff < data.size){
+    result = data.str + foff;
+  }
+  return(result);
+}
+
+SYMS_API void*
+syms_pe_ptr_from_voff(SYMS_String8 data, SYMS_PeBinAccel *bin, SYMS_U64 voff){
+  //- get the section for this voff
+  SYMS_U64 sec_count = bin->section_count;
+  SYMS_CoffSection *sec_array = (SYMS_CoffSection*)((SYMS_U8*)data.str + bin->section_array_off);
+  SYMS_CoffSection *sec_ptr = sec_array;
+  SYMS_CoffSection *sec = 0;
+  for (SYMS_U64 i = 1; i <= sec_count; i += 1, sec_ptr += 1){
+    if (sec_ptr->virt_off <= voff && voff < sec_ptr->virt_off + sec_ptr->virt_size){
+      sec = sec_ptr;
+      break;
+    }
+  }
+  
+  //- adjust to file pointer
+  void *result = 0;
+  if (sec != 0 & sec_ptr->file_size > 0){
+    result = data.str + voff - sec->virt_off + sec->file_off;
+  }
+  
+  return(result);
+}
+
+SYMS_API SYMS_U64
+syms_pe_virt_off_to_file_off(SYMS_CoffSection *sections, SYMS_U32 section_count, SYMS_U64 virt_off){
+  SYMS_U64 file_off = 0;
+  for (SYMS_U32 sect_idx = 0; sect_idx < section_count; sect_idx += 1){
+    SYMS_CoffSection *sect = &sections[sect_idx];
+    if (sect->virt_off <= virt_off && virt_off < sect->virt_off + sect->virt_size){
+      file_off = sect->file_off + (virt_off - sect->virt_off);
+    }
+  }
+  return file_off;
+}
+
+SYMS_API SYMS_U64
+syms_pe_bin_virt_off_to_file_off(SYMS_String8 data, SYMS_PeBinAccel *bin, SYMS_U64 virt_off){
+  return syms_pe_virt_off_to_file_off((SYMS_CoffSection*)((SYMS_U8*)data.str + bin->section_array_off), bin->section_count, virt_off);
+}
+
 #endif //SYMS_PE_PARSER_C
diff --git a/Engine/Source/Runtime/SymsLib/syms/core/pe/syms_pe_parser.h b/Engine/Source/Runtime/SymsLib/syms/core/pe/syms_pe_parser.h
--- a/Engine/Source/Runtime/SymsLib/syms/core/pe/syms_pe_parser.h
+++ b/Engine/Source/Runtime/SymsLib/syms/core/pe/syms_pe_parser.h
@@ -4,6 +4,8 @@
 #ifndef SYMS_PE_PARSER_H
 #define SYMS_PE_PARSER_H
 
+// TODO(allen): rename? split PE and COFF accelerators?
+
 ////////////////////////////////
 // NOTE(allen): PE Parser Types
 
@@ -12,6 +14,11 @@ typedef struct SYMS_PeFileAccel{
   SYMS_U32 coff_off;
 } SYMS_PeFileAccel;
 
+typedef struct SYMS_CoffFileAccel{
+  SYMS_FileFormat format;
+  // TODO(allen): 
+} SYMS_CoffFileAccel;
+
 typedef struct SYMS_PeBinAccel{
   SYMS_FileFormat format;
   SYMS_U64 image_base;
@@ -27,11 +34,15 @@ typedef struct SYMS_PeBinAccel{
   SYMS_U64Range *data_dirs_file;
 } SYMS_PeBinAccel;
 
-////////////////////////////////
-// Imports
+typedef struct SYMS_CoffBinAccel{
+  SYMS_FileFormat format;
+  SYMS_U64 image_base;
+  SYMS_U64 section_array_off;
+  SYMS_U64 section_count;
+  SYMS_Arch arch;
+} SYMS_CoffBinAccel;
 
-enum
-{
+enum{
   SYMS_PeLookupBy_Ordinal,
   SYMS_PeLookupBy_NameHint
 };
@@ -64,8 +75,14 @@ SYMS_API SYMS_PeFileAccel* syms_pe_file_accel_from_data(SYMS_Arena *arena, SYMS_
 SYMS_API SYMS_PeBinAccel*  syms_pe_bin_accel_from_file(SYMS_Arena *arena, SYMS_String8 data,
                                                        SYMS_PeFileAccel *file);
 
+SYMS_API SYMS_CoffFileAccel* syms_coff_file_accel_from_data(SYMS_Arena *arena, SYMS_String8 data);
+SYMS_API SYMS_CoffBinAccel*  syms_coff_bin_accel_from_file(SYMS_Arena *arena, SYMS_String8 data,
+                                                           SYMS_CoffFileAccel *file);
+
+
 // arch
 SYMS_API SYMS_Arch syms_pe_arch_from_bin(SYMS_PeBinAccel *bin);
+SYMS_API SYMS_Arch syms_coff_arch_from_bin(SYMS_CoffBinAccel *bin);
 
 // external info
 SYMS_API SYMS_ExtFileList syms_pe_ext_file_list_from_bin(SYMS_Arena *arena, SYMS_String8 data,
@@ -75,9 +92,20 @@ SYMS_API SYMS_ExtFileList syms_pe_ext_file_list_from_bin(SYMS_Arena *arena, SYMS
 //  pe specific
 SYMS_API SYMS_CoffSection syms_pe_coff_section(SYMS_String8 data, SYMS_PeBinAccel *bin, SYMS_U64 n);
 
+SYMS_API SYMS_SecInfoArray syms_pe_coff_sec_info_array_from_data(SYMS_Arena *arena, SYMS_String8 data,
+                                                                 SYMS_U64 sec_array_off, SYMS_U64 sec_count);
+
+
 //  main api
 SYMS_API SYMS_SecInfoArray syms_pe_sec_info_array_from_bin(SYMS_Arena *arena, SYMS_String8 data,
                                                            SYMS_PeBinAccel *bin);
+SYMS_API SYMS_SecInfoArray syms_coff_sec_info_array_from_bin(SYMS_Arena *arena, SYMS_String8 data,
+                                                             SYMS_CoffBinAccel *bin);
+
+// imports & exports
+//  main api
+SYMS_API SYMS_ImportArray syms_pe_imports_from_bin(SYMS_Arena *arena, SYMS_String8 data, SYMS_PeBinAccel *bin);
+SYMS_API SYMS_ExportArray syms_pe_exports_from_bin(SYMS_Arena *arena, SYMS_String8 data, SYMS_PeBinAccel *bin);
 
 ////////////////////////////////
 // NOTE(allen): PE Specific Helpers
@@ -93,9 +121,4 @@ SYMS_API void*    syms_pe_ptr_from_voff(SYMS_String8 data, SYMS_PeBinAccel *bin,
 SYMS_API SYMS_U64 syms_pe_virt_off_to_file_off(SYMS_CoffSection *sections, SYMS_U32 section_count, SYMS_U64 virt_off);
 SYMS_API SYMS_U64 syms_pe_bin_virt_off_to_file_off(SYMS_String8 data, SYMS_PeBinAccel *bin, SYMS_U64 virt_off);
 
-////////////////////////////////
-
-SYMS_API SYMS_ImportArray syms_pe_imports_from_bin(SYMS_Arena *arena, SYMS_String8 data, SYMS_PeBinAccel *bin);
-SYMS_API SYMS_ExportArray syms_pe_exports_from_bin(SYMS_Arena *arena, SYMS_String8 data, SYMS_PeBinAccel *bin);
-
 #endif //SYMS_PE_PARSER_H
diff --git a/Engine/Source/Runtime/SymsLib/syms/core/syms_debug_info.h b/Engine/Source/Runtime/SymsLib/syms/core/syms_debug_info.h
--- a/Engine/Source/Runtime/SymsLib/syms/core/syms_debug_info.h
+++ b/Engine/Source/Runtime/SymsLib/syms/core/syms_debug_info.h
@@ -5,7 +5,7 @@
 #define SYMS_DEBUG_INFO_H
 
 ////////////////////////////////
-//~ NOTE(allen): Generated Types
+//~ allen: Generated Types
 
 #include "syms/core/generated/syms_meta_debug_info.h"
 
@@ -34,7 +34,7 @@ typedef struct SYMS_BinInfoArray{
 } SYMS_BinInfoArray;
 
 ////////////////////////////////
-//~ NOTE(allen): Binary Imports & Exports
+//~ allen: Binary Imports & Exports
 
 typedef struct SYMS_Import{
   SYMS_String8 name;
@@ -72,7 +72,7 @@ typedef struct SYMS_ExportArray{
 } SYMS_ExportArray;
 
 ////////////////////////////////
-//~ NOTE(allen): Compilation Unit Types
+//~ allen: Compilation Unit Types
 
 typedef SYMS_U64 SYMS_UnitID;
 
@@ -100,7 +100,7 @@ typedef struct SYMS_UnitRangeArray{
 } SYMS_UnitRangeArray;
 
 ////////////////////////////////
-//~ NOTE(allen): External File Types
+//~ allen: External File Types
 
 typedef struct SYMS_ExtMatchKey{
   SYMS_U8 v[16];
@@ -123,7 +123,7 @@ typedef struct SYMS_ExtFileList{
 } SYMS_ExtFileList;
 
 ////////////////////////////////
-//~ NOTE(allen): Line Info Types
+//~ allen: Line Info Types
 
 typedef SYMS_U64 SYMS_FileID;
 
@@ -138,6 +138,13 @@ typedef struct SYMS_Line{
   SYMS_U64 voff;
 } SYMS_Line;
 
+typedef struct SYMS_ResolvedLine{
+  SYMS_String8 file_name;
+  SYMS_U32 line;
+  SYMS_U32 col;
+  SYMS_U64 voff;
+} SYMS_ResolvedLine;
+
 typedef struct SYMS_FileIDArray{
   SYMS_FileID *ids;
   SYMS_U64 count;
@@ -161,7 +168,7 @@ typedef struct SYMS_LineParseOut{
 } SYMS_LineParseOut;
 
 ////////////////////////////////
-//~ NOTE(allen): Type Info Types
+//~ allen: Type Info Types
 
 // TODO(allen): go to the mc file
 // TODO(allen): just yanked right from PDB
@@ -203,7 +210,7 @@ typedef enum SYMS_SizeInterpretation{
 } SYMS_SizeInterpretation;
 
 ////////////////////////////////
-//~ NOTE(allen): Symbol API Types
+//~ allen: Symbol API Types
 
 typedef SYMS_U64 SYMS_SymbolID;
 
@@ -296,16 +303,6 @@ typedef struct SYMS_ConstInfo{
   SYMS_U64 val;
 } SYMS_ConstInfo;
 
-typedef struct SYMS_StrippedInfo{
-  SYMS_String8 name;
-  SYMS_U64 voff;
-} SYMS_StrippedInfo;
-
-typedef struct SYMS_StrippedInfoArray{
-  SYMS_StrippedInfo *info;
-  SYMS_U64 count;
-} SYMS_StrippedInfoArray;
-
 typedef enum SYMS_MemKind{
   SYMS_MemKind_Null,
   SYMS_MemKind_DataField,
@@ -352,7 +349,7 @@ typedef struct SYMS_EnumInfoArray{
 } SYMS_EnumInfoArray;
 
 ////////////////////////////////
-// NOTE(allen): Location Types
+// allen: Location Types
 
 typedef SYMS_U64 SYMS_LocID;
 
@@ -383,12 +380,25 @@ typedef enum SYMS_ProcLoc{
 } SYMS_ProcLoc;
 
 ////////////////////////////////
-// NOTE(allen): Nil For Accelerators
+// allen: Link Name Types
+
+typedef struct SYMS_LinkNameRec{
+  SYMS_String8 name;
+  SYMS_U64 voff;
+} SYMS_LinkNameRec;
+
+typedef struct SYMS_LinkNameRecArray{
+  SYMS_LinkNameRec *recs;
+  SYMS_U64 count;
+} SYMS_LinkNameRecArray;
+
+////////////////////////////////
+// allen: Nil For Accelerators
 
 SYMS_READ_ONLY SYMS_GLOBAL SYMS_FileFormat syms_format_nil = SYMS_FileFormat_Null;
 
 ////////////////////////////////
-//~ NOTE(allen): Symbol Helper Functions
+//~ allen: Symbol Helper Functions
 
 SYMS_C_LINKAGE_BEGIN
 
diff --git a/Engine/Source/Runtime/SymsLib/syms/core/syms_parser.c b/Engine/Source/Runtime/SymsLib/syms/core/syms_parser.c
--- a/Engine/Source/Runtime/SymsLib/syms/core/syms_parser.c
+++ b/Engine/Source/Runtime/SymsLib/syms/core/syms_parser.c
@@ -9,6 +9,8 @@
 SYMS_API SYMS_FileAccel*
 syms_file_accel_from_data(SYMS_Arena *arena, SYMS_String8 data){
   SYMS_FileAccel *result = (SYMS_FileAccel*)&syms_format_nil;
+  
+  // easy to recognize
   if (result->format == SYMS_FileFormat_Null){
     result = (SYMS_FileAccel*)syms_pe_file_accel_from_data(arena, data);
   }
@@ -21,6 +23,12 @@ syms_file_accel_from_data(SYMS_Arena *arena, SYMS_String8 data){
   if (result->format == SYMS_FileFormat_Null){
     result = (SYMS_FileAccel*)syms_mach_file_accel_from_data(arena, data);
   }
+  
+  // harder to recognize
+  if (result->format == SYMS_FileFormat_Null){
+    result = (SYMS_FileAccel*)syms_coff_file_accel_from_data(arena, data);
+  }
+  
   SYMS_ASSERT_PARANOID(syms_parser_api_invariants());
   return(result);
 }
@@ -40,6 +48,8 @@ syms_file_is_bin(SYMS_FileAccel *file){
   switch (file->format){
     case SYMS_FileFormat_PE:
     case SYMS_FileFormat_ELF:
+    case SYMS_FileFormat_COFF16:
+    case SYMS_FileFormat_COFF32:
     {
       result = syms_true;
     }break;
@@ -63,6 +73,12 @@ syms_bin_accel_from_file(SYMS_Arena *arena, SYMS_String8 data, SYMS_FileAccel *f
       result = (SYMS_BinAccel*)syms_pe_bin_accel_from_file(arena, data, &file->pe_accel);
     }break;
     
+    case SYMS_FileFormat_COFF16:
+    case SYMS_FileFormat_COFF32:
+    {
+      result = (SYMS_BinAccel*)syms_coff_bin_accel_from_file(arena, data, &file->coff_accel);
+    }break;
+    
     case SYMS_FileFormat_ELF:
     {
       result = (SYMS_BinAccel*)syms_elf_bin_accel_from_file(arena, data, &file->elf_accel);
@@ -89,6 +105,12 @@ syms_arch_from_bin(SYMS_BinAccel *bin){
       result = syms_pe_arch_from_bin(&bin->pe_accel);
     }break;
     
+    case SYMS_FileFormat_COFF16:
+    case SYMS_FileFormat_COFF32:
+    {
+      result = syms_coff_arch_from_bin(&bin->coff_accel);
+    }break;
+    
     case SYMS_FileFormat_ELF:
     {
       result = syms_elf_arch_from_bin(&bin->elf_accel);
@@ -136,6 +158,12 @@ syms_sec_info_array_from_bin(SYMS_Arena *arena, SYMS_String8 data, SYMS_BinAccel
       result = syms_pe_sec_info_array_from_bin(arena, data, &bin->pe_accel);
     }break;
     
+    case SYMS_FileFormat_COFF16:
+    case SYMS_FileFormat_COFF32:
+    {
+      result = syms_coff_sec_info_array_from_bin(arena, data, &bin->coff_accel);
+    }break;
+    
     case SYMS_FileFormat_ELF:
     {
       result = syms_elf_sec_info_array_from_bin(arena, data, &bin->elf_accel);
@@ -167,6 +195,11 @@ syms_default_vbase_from_bin(SYMS_BinAccel *bin){
     {
       result = syms_elf_default_vbase_from_bin(&bin->elf_accel);
     }break;
+    
+    case SYMS_FileFormat_MACH:
+    {
+      // TODO(allen): ?
+    }break;
   }
   SYMS_ASSERT_PARANOID(syms_parser_api_invariants());
   SYMS_ProfEnd();
@@ -1233,28 +1266,6 @@ syms_location_ranges_from_proc_sid(SYMS_Arena *arena, SYMS_String8 data, SYMS_Db
   return(result);
 }
 
-SYMS_API SYMS_StrippedInfoArray
-syms_stripped_from_unit(SYMS_Arena *arena, SYMS_String8 data, SYMS_DbgAccel *dbg, SYMS_UnitAccel *unit){
-  SYMS_ProfBegin("syms_stripped_from_unit");
-  SYMS_StrippedInfoArray result = {0};
-  if (unit->format == dbg->format){
-    switch (unit->format){
-      case SYMS_FileFormat_PDB:
-      {
-        result = syms_pdb_stripped_from_unit(arena, data, &dbg->pdb_accel, &unit->pdb_accel);
-      }break;
-      
-      case SYMS_FileFormat_DWARF:
-      {
-        // do nothing
-      }break;
-    }
-  }
-  SYMS_ASSERT_PARANOID(syms_parser_api_invariants());
-  SYMS_ProfEnd();
-  return(result);
-}
-
 // signature info
 SYMS_API SYMS_SigInfo
 syms_sig_info_from_type_sid(SYMS_Arena *arena, SYMS_String8 data, SYMS_DbgAccel *dbg,
@@ -1369,13 +1380,13 @@ syms_type_map_from_dbg(SYMS_Arena *arena, SYMS_String8 data, SYMS_DbgAccel *dbg)
 }
 
 SYMS_API SYMS_MapAccel*
-syms_image_symbol_map_from_dbg(SYMS_Arena *arena, SYMS_String8 data, SYMS_DbgAccel *dbg){
-  SYMS_ProfBegin("syms_image_symbol_map_from_dbg");
+syms_unmangled_symbol_map_from_dbg(SYMS_Arena *arena, SYMS_String8 data, SYMS_DbgAccel *dbg){
+  SYMS_ProfBegin("syms_unmangled_symbol_map_from_dbg");
   SYMS_MapAccel *result = (SYMS_MapAccel*)&syms_format_nil;
   switch (dbg->format){
     case SYMS_FileFormat_PDB:
     {
-      result = (SYMS_MapAccel*)syms_pdb_image_symbol_map_from_dbg(arena, data, &dbg->pdb_accel);
+      result = (SYMS_MapAccel*)syms_pdb_unmangled_symbol_map_from_dbg(arena, data, &dbg->pdb_accel);
     }break;
     
     case SYMS_FileFormat_DWARF:
@@ -1409,23 +1420,23 @@ syms_partner_uid_from_map(SYMS_MapAccel *map){
 }
 
 SYMS_API SYMS_USIDList
-syms_symbol_list_from_string(SYMS_Arena *arena, SYMS_String8 data, SYMS_DbgAccel *dbg,
-                             SYMS_MapAndUnit map_and_unit, SYMS_String8 string){
-  SYMS_ProfBegin("syms_symbol_list_from_string");
+syms_usid_list_from_string(SYMS_Arena *arena, SYMS_String8 data, SYMS_DbgAccel *dbg,
+                           SYMS_MapAndUnit *map_and_unit, SYMS_String8 string){
+  SYMS_ProfBegin("syms_usid_list_from_string");
   SYMS_USIDList result = {0};
-  if (map_and_unit.map->format == dbg->format){
+  if (map_and_unit->map->format == dbg->format){
     switch (dbg->format){
       case SYMS_FileFormat_PDB:
       {
         result = syms_pdb_usid_list_from_string(arena, data, &dbg->pdb_accel,
-                                                &map_and_unit.unit->pdb_accel,
-                                                &map_and_unit.map->pdb_accel,
+                                                &map_and_unit->unit->pdb_accel,
+                                                &map_and_unit->map->pdb_accel,
                                                 string);
       }break;
       
       case SYMS_FileFormat_DWARF:
       {
-        result = syms_dw_usid_list_from_string(arena, &map_and_unit.map->dw_accel, string);
+        result = syms_dw_usid_list_from_string(arena, &map_and_unit->map->dw_accel, string);
       }break;
     }
   }
@@ -1434,6 +1445,106 @@ syms_symbol_list_from_string(SYMS_Arena *arena, SYMS_String8 data, SYMS_DbgAccel
   return(result);
 }
 
+// mangled names (linker names)
+
+SYMS_API SYMS_UnitID
+syms_link_names_uid(SYMS_DbgAccel *dbg){
+  SYMS_UnitID result = 0;
+  switch (dbg->format){
+    case SYMS_FileFormat_PDB:
+    {
+      result = syms_pdb_link_names_uid();
+    }break;
+    
+    case SYMS_FileFormat_DWARF:
+    {
+      // NOTE(allen): do nothing
+    }break;
+  }
+  return(result);
+}
+
+SYMS_API SYMS_B32
+syms_link_map_is_complete(SYMS_LinkMapAccel *map){
+  SYMS_ProfBegin("syms_link_map_is_complete");
+  SYMS_B32 result = syms_false;
+  switch (map->format){
+    case SYMS_FileFormat_PDB:
+    {
+      result = syms_true;
+    }break;
+    case SYMS_FileFormat_DWARF:
+    {
+      result = syms_false;
+    }break;
+  }
+  SYMS_ProfEnd();
+  return(result);
+}
+
+SYMS_API SYMS_LinkMapAccel*
+syms_link_map_from_dbg(SYMS_Arena *arena, SYMS_String8 data, SYMS_DbgAccel *dbg){
+  SYMS_ProfBegin("syms_link_map_from_dbg");
+  SYMS_LinkMapAccel *result = (SYMS_LinkMapAccel*)&syms_format_nil;
+  switch (dbg->format){
+    case SYMS_FileFormat_PDB:
+    {
+      result = (SYMS_LinkMapAccel*)syms_pdb_link_map_from_dbg(arena, data, (SYMS_PdbDbgAccel*)dbg);
+    }break;
+    
+    case SYMS_FileFormat_DWARF:
+    {
+      // TODO(allen): 
+    }break;
+  }
+  SYMS_ProfEnd();
+  return(result);
+}
+
+SYMS_API SYMS_U64
+syms_voff_from_link_name(SYMS_String8 data, SYMS_DbgAccel *dbg, SYMS_LinkMapAccel *map,
+                         SYMS_UnitAccel *link_unit, SYMS_String8 name){
+  SYMS_ProfBegin("syms_voff_from_link_name");
+  SYMS_U64 result = 0;
+  if (dbg->format == map->format){
+    switch (dbg->format){
+      case SYMS_FileFormat_PDB:
+      {
+        result = syms_pdb_voff_from_link_name(data, (SYMS_PdbDbgAccel*)dbg, (SYMS_PdbLinkMapAccel*)map,
+                                              (SYMS_PdbUnitAccel*)link_unit, name);
+      }break;
+      
+      case SYMS_FileFormat_DWARF:
+      {
+        // TODO(allen): 
+      }break;
+    }
+  }
+  SYMS_ProfEnd();
+  return(result);
+}
+
+SYMS_API SYMS_LinkNameRecArray
+syms_link_name_array_from_unit(SYMS_Arena *arena, SYMS_String8 data, SYMS_DbgAccel *dbg, SYMS_UnitAccel *unit){
+  SYMS_ProfBegin("syms_link_name_array_from_unit");
+  SYMS_LinkNameRecArray result = {0};
+  if (dbg->format == unit->format){
+    switch (dbg->format){
+      case SYMS_FileFormat_PDB:
+      {
+        result = syms_pdb_link_name_array_from_unit(arena, data, (SYMS_PdbDbgAccel*)dbg, (SYMS_PdbUnitAccel*)unit);
+      }break;
+      
+      case SYMS_FileFormat_DWARF:
+      {
+        // TODO(allen): 
+      }break;
+    }
+  }
+  SYMS_ProfEnd();
+  return(result);
+}
+
 // thread vars
 
 SYMS_API SYMS_UnitID
diff --git a/Engine/Source/Runtime/SymsLib/syms/core/syms_parser.h b/Engine/Source/Runtime/SymsLib/syms/core/syms_parser.h
--- a/Engine/Source/Runtime/SymsLib/syms/core/syms_parser.h
+++ b/Engine/Source/Runtime/SymsLib/syms/core/syms_parser.h
@@ -5,12 +5,13 @@
 #define SYMS_PARSER_H
 
 ////////////////////////////////
-//~ NOTE(allen): Abstraction Acclerators
+//~ allen: Abstraction Acclerators
 
 typedef union SYMS_FileAccel{
   SYMS_FileFormat format;
   SYMS_ElfFileAccel elf_accel;
   SYMS_PeFileAccel pe_accel;
+  SYMS_CoffFileAccel coff_accel;
   SYMS_PdbFileAccel pdb_accel;
   SYMS_MachFileAccel mach_accel;
 } SYMS_FileAccel;
@@ -18,6 +19,7 @@ typedef union SYMS_FileAccel{
 typedef union SYMS_BinAccel{
   SYMS_FileFormat format;
   SYMS_PeBinAccel pe_accel;
+  SYMS_CoffBinAccel coff_accel;
   SYMS_ElfBinAccel elf_accel;
   SYMS_MachBinAccel mach_accel;
 } SYMS_BinAccel;
@@ -57,8 +59,14 @@ typedef union SYMS_MapAccel{
   SYMS_PdbMapAccel pdb_accel;
 } SYMS_MapAccel;
 
+typedef union SYMS_LinkMapAccel{
+  SYMS_FileFormat format;
+  SYMS_DwLinkMapAccel dw_accel;
+  SYMS_PdbLinkMapAccel pdb_accel;
+} SYMS_LinkMapAccel;
+
 ////////////////////////////////
-//~ NOTE(allen): Accelerator Bundle Types
+//~ allen: Accelerator Bundle Types
 
 typedef struct SYMS_MapAndUnit{
   SYMS_MapAccel *map;
@@ -76,18 +84,18 @@ SYMS_C_LINKAGE_BEGIN
 
 
 ////////////////////////////////
-//~ NOTE(allen): Accel Helpers
+//~ allen: Accel Helpers
 
 #define syms_accel_is_good(a) ((a) != 0 && (a)->format != SYMS_FileFormat_Null)
 
 ////////////////////////////////
-//~ NOTE(allen): General File Analysis
+//~ allen: General File Analysis
 
 SYMS_API SYMS_FileAccel*   syms_file_accel_from_data(SYMS_Arena *arena, SYMS_String8 data);
 SYMS_API SYMS_FileFormat   syms_file_format_from_file(SYMS_FileAccel *file);
 
 ////////////////////////////////
-//~ NOTE(allen): Bin File
+//~ allen: Bin File
 
 SYMS_API SYMS_B32          syms_file_is_bin(SYMS_FileAccel *file);
 SYMS_API SYMS_BinAccel*    syms_bin_accel_from_file(SYMS_Arena *arena, SYMS_String8 data, SYMS_FileAccel *file);
@@ -109,7 +117,7 @@ SYMS_API SYMS_ImportArray  syms_imports_from_bin(SYMS_Arena *arena, SYMS_String8
 SYMS_API SYMS_ExportArray  syms_exports_from_bin(SYMS_Arena *arena, SYMS_String8 data, SYMS_BinAccel *bin);
 
 ////////////////////////////////
-//~ NOTE(nick): Bin List
+//~ nick: Bin List
 
 SYMS_API SYMS_B32          syms_file_is_bin_list(SYMS_FileAccel *file);
 SYMS_API SYMS_BinListAccel*syms_bin_list_from_file(SYMS_Arena *arena, SYMS_String8 data, SYMS_FileAccel *file);
@@ -118,7 +126,7 @@ SYMS_API SYMS_BinAccel*    syms_bin_accel_from_bin_list_number(SYMS_Arena *arena
                                                                SYMS_BinListAccel *list, SYMS_U64 n);
 
 ////////////////////////////////
-//~ NOTE(rjf): Dbg File
+//~ rjf: Dbg File
 
 SYMS_API SYMS_B32          syms_file_is_dbg(SYMS_FileAccel *file);
 SYMS_API SYMS_DbgAccel*    syms_dbg_accel_from_file(SYMS_Arena *arena, SYMS_String8 data, SYMS_FileAccel *file);
@@ -207,7 +215,7 @@ SYMS_API SYMS_SigInfo      syms_sig_info_from_mem_number(SYMS_Arena *arena, SYMS
 SYMS_API SYMS_USID         syms_symbol_from_mem_number(SYMS_String8 data, SYMS_DbgAccel *dbg, SYMS_UnitAccel *unit, SYMS_MemsAccel *mems, SYMS_U64 n);
 SYMS_API SYMS_USID         syms_containing_type_from_sid(SYMS_String8 data, SYMS_DbgAccel *dbg,
                                                          SYMS_UnitAccel *unit, SYMS_SymbolID sid);
-// TODO(rjf): linkage name
+// TODO(rjf): linkage name from symbol
 
 SYMS_API SYMS_EnumInfoArray syms_enum_info_array_from_sid(SYMS_Arena *arena, SYMS_String8 data, SYMS_DbgAccel *dbg,
                                                           SYMS_UnitAccel *unit, SYMS_SymbolID sid);
@@ -233,9 +241,6 @@ SYMS_API SYMS_LocRangeArray syms_location_ranges_from_proc_sid(SYMS_Arena *arena
                                                                SYMS_UnitAccel *unit, SYMS_SymbolID sid,
                                                                SYMS_ProcLoc proc_loc);
 
-SYMS_API SYMS_StrippedInfoArray syms_stripped_from_unit(SYMS_Arena *arena, SYMS_String8 data, SYMS_DbgAccel *dbg,
-                                                        SYMS_UnitAccel *unit);
-
 // signature info
 SYMS_API SYMS_SigInfo      syms_sig_info_from_type_sid(SYMS_Arena *arena, SYMS_String8 data, SYMS_DbgAccel *dbg,
                                                        SYMS_UnitAccel *unit, SYMS_SymbolID sid);
@@ -252,13 +257,22 @@ SYMS_API SYMS_LineParseOut syms_line_parse_from_uid(SYMS_Arena *arena, SYMS_Stri
 
 // name maps
 SYMS_API SYMS_MapAccel*    syms_type_map_from_dbg(SYMS_Arena *arena, SYMS_String8 data, SYMS_DbgAccel *dbg);
-SYMS_API SYMS_MapAccel*    syms_image_symbol_map_from_dbg(SYMS_Arena *arena,SYMS_String8 data,SYMS_DbgAccel *dbg);
+SYMS_API SYMS_MapAccel*    syms_unmangled_symbol_map_from_dbg(SYMS_Arena *arena,SYMS_String8 data,SYMS_DbgAccel *dbg);
 SYMS_API SYMS_UnitID       syms_partner_uid_from_map(SYMS_MapAccel *map);
-SYMS_API SYMS_USIDList     syms_symbol_list_from_string(SYMS_Arena *arena, SYMS_String8 data, SYMS_DbgAccel *dbg,
-                                                        SYMS_MapAndUnit map_and_unit, SYMS_String8 string);
+SYMS_API SYMS_USIDList     syms_usid_list_from_string(SYMS_Arena *arena, SYMS_String8 data, SYMS_DbgAccel *dbg,
+                                                      SYMS_MapAndUnit *map_and_unit, SYMS_String8 string);
 
-// thread vars
+// linker names
+SYMS_API SYMS_UnitID       syms_link_names_uid(SYMS_DbgAccel *dbg);
+
+SYMS_API SYMS_LinkMapAccel*syms_link_map_from_dbg(SYMS_Arena *arena, SYMS_String8 data, SYMS_DbgAccel *dbg);
+SYMS_API SYMS_B32          syms_link_map_is_complete(SYMS_LinkMapAccel *map);
+SYMS_API SYMS_U64          syms_voff_from_link_name(SYMS_String8 data, SYMS_DbgAccel *dbg, SYMS_LinkMapAccel *map,
+                                                    SYMS_UnitAccel *link_unit, SYMS_String8 name);
+SYMS_API SYMS_LinkNameRecArray syms_link_name_array_from_unit(SYMS_Arena *arena, SYMS_String8 data,
+                                                              SYMS_DbgAccel *dbg, SYMS_UnitAccel *unit);
 
+// thread vars
 SYMS_API SYMS_UnitID        syms_tls_var_uid_from_dbg(SYMS_DbgAccel *dbg);
 SYMS_API SYMS_SymbolIDArray syms_tls_var_sid_array_from_unit(SYMS_Arena *arena, SYMS_UnitAccel *thread_unit);
 
