Fix off by one error in section count for elf parsing

diff --git a/Engine/Source/Runtime/SymsLib/syms/core/elf/syms_elf_parser.c b/Engine/Source/Runtime/SymsLib/syms/core/elf/syms_elf_parser.c
--- a/Engine/Source/Runtime/SymsLib/syms/core/elf/syms_elf_parser.c
+++ b/Engine/Source/Runtime/SymsLib/syms/core/elf/syms_elf_parser.c
@@ -160,21 +160,18 @@ syms_elf_section_array_from_img_header(SYMS_Arena *arena, SYMS_String8 file, SYM
   SYMS_U64Range file_range = syms_make_u64_range(0, file.size);
   
   //- rjf: figure out section count
-  SYMS_U64 section_count = img.ehdr.e_shnum;
+  // NOTE(rjf): ELF files have a null/empty section at the first slot of the
+  // section header table. We're explicitly skipping that, so we need to
+  // account for (e_shnum-1) sections.
+  SYMS_U64 section_count = img.ehdr.e_shnum ? (img.ehdr.e_shnum-1) : 0;
   
   //- rjf: figure out section range and section header size (32-bit or 64-bit)
   SYMS_U64Range section_range = syms_make_u64_range(SYMS_U64_MAX, SYMS_U64_MAX);
-  SYMS_U64 section_header_size = SYMS_U64_MAX;
+  SYMS_U64 section_header_size = img.ehdr.e_shentsize;
   {
-    SYMS_U64 section_range_base = img.ehdr.e_shoff+img.ehdr.e_shentsize;
-    switch(img.is_32bit)
-    {
-      case syms_true:  section_header_size = sizeof(SYMS_ElfShdr32); break;
-      case syms_false: section_header_size = sizeof(SYMS_ElfShdr64); break;
-      default: break;
-    }
-    section_range.min = section_range_base;
-    section_range.max = section_range_base+img.ehdr.e_shnum*section_header_size;
+    section_range.min = img.ehdr.e_shoff + 1*section_header_size;
+    section_range.max = section_range.min + section_count*section_header_size;
+    section_range.max = SYMS_ClampTop(file.size, section_range.max);
   }
   
   //- rjf: allocate sections
@@ -223,13 +220,6 @@ syms_elf_section_array_from_img_header(SYMS_Arena *arena, SYMS_String8 file, SYM
   result.v = sections;
   result.count = section_count;
   
-  //- rjf: if the last section is empty, subtract 1 from the count so we do not report it
-  // to the user as real information
-  if(result.count > 0 && result.v[result.count-1].name.size == 0)
-  {
-    result.count -= 1;
-  }
-  
   return result;
 }
 
@@ -238,7 +228,7 @@ syms_elf_segment_array_from_img_header(SYMS_Arena *arena, SYMS_String8 file, SYM
 {
   void *base = file.str;
   SYMS_U64Range range = syms_make_u64_range(0, file.size);
-
+  
   SYMS_U64 segment_count = img.ehdr.e_phnum;
   SYMS_ElfPhdr64 *segments = syms_push_array_zero(arena, SYMS_ElfPhdr64, segment_count);
   for(SYMS_U64 segment_idx = 0; segment_idx < segment_count; segment_idx += 1) 
@@ -254,11 +244,11 @@ syms_elf_segment_array_from_img_header(SYMS_Arena *arena, SYMS_String8 file, SYM
       syms_based_range_read_struct(base, range, img.ehdr.e_phoff + segment_idx * sizeof(SYMS_ElfPhdr64), &segments[segment_idx]);
     }
   }
-
+  
   SYMS_ElfSegmentArray result;
   result.count = segment_count;
   result.v = segments;
-
+  
   return result;
 }
 
diff --git a/Engine/Source/Runtime/SymsLib/syms/core/group/syms_group.h b/Engine/Source/Runtime/SymsLib/syms/core/group/syms_group.h
--- a/Engine/Source/Runtime/SymsLib/syms/core/group/syms_group.h
+++ b/Engine/Source/Runtime/SymsLib/syms/core/group/syms_group.h
@@ -151,7 +151,7 @@ typedef struct SYMS_Group{
 
 SYMS_READ_ONLY SYMS_GLOBAL SYMS_SymbolIDArray syms_sid_array_nil = {0};
 SYMS_READ_ONLY SYMS_GLOBAL SYMS_String8Array syms_string_array_nil = {0};
-SYMS_READ_ONLY SYMS_GLOBAL SYMS_LineParseOut syms_line_parse_nil = {0};
+SYMS_READ_ONLY SYMS_GLOBAL SYMS_LineParseOut syms_line_parse_nil = {{0}};
 SYMS_READ_ONLY SYMS_GLOBAL SYMS_SpatialMap1D syms_spatial_map_1d_nil = {0};
 SYMS_READ_ONLY SYMS_GLOBAL SYMS_LineToAddrMap syms_line_to_addr_map_nil = {0};
 SYMS_READ_ONLY SYMS_GLOBAL SYMS_TypeMemberArray syms_type_member_array_nil = {0};
