Add SymsLib support for MacOS

diff --git a/Engine/Source/Runtime/SymsLib/BuildForMac.sh b/Engine/Source/Runtime/SymsLib/BuildForMac.sh
new file mode 100755
--- /dev/null
+++ b/Engine/Source/Runtime/SymsLib/BuildForMac.sh
@@ -0,0 +1,50 @@
+#!/bin/bash
+
+set -eu
+
+SCRIPT_DIR=$(cd "$(dirname "$BASH_SOURCE")" ; pwd)
+UE_THIRD_PARTY_DIR=$(cd "${SCRIPT_DIR}/../../../.."; pwd)
+SYMS_DIR=$(cd "${SCRIPT_DIR}"; pwd)
+
+CFLAGS="-ffunction-sections -fdata-sections"
+CXXFLAGS="-ffunction-sections -fdata-sections"
+LIBS="-lm -lc -lpthread"
+
+# Get num of cores
+export CORES=$(getconf _NPROCESSORS_ONLN)
+echo "Using ${CORES} cores for building"
+
+BuildLibSyms()
+{
+    export ARCH=$1
+    export FLAVOR=$2
+    local EXT=$3
+    local BUILD_DIR=/tmp/Build-Syms-${ARCH}-${FLAVOR}
+
+    echo "Building ${ARCH} ${FLAVOR}"
+    rm -rf ${BUILD_DIR}
+    mkdir -p ${BUILD_DIR}/Build
+
+    pushd ${BUILD_DIR}/Build
+
+    set -x
+    cmake \
+      -DCMAKE_MAKE_PROGRAM=$(which make) \
+      -DCMAKE_BUILD_TYPE=${FLAVOR} \
+      -DCMAKE_POSITION_INDEPENDENT_CODE=ON \
+      -DCMAKE_OSX_DEPLOYMENT_TARGET=11.0 \
+      ${SYMS_DIR}/
+    set +x
+    
+    VERBOSE=1 make -j ${CORES}
+
+    mkdir -p ${SYMS_DIR}/lib/Mac/${ARCH}
+    cp -v ${BUILD_DIR}/Build/libsyms.a ${SYMS_DIR}/lib/Mac/${ARCH}/libsyms${EXT}.a
+
+    popd
+}
+
+BuildLibSyms x86_64 Release ""
+BuildLibSyms x86_64 Debug "d"
+BuildLibSyms arm64 Release ""
+BuildLibSyms arm64 Debug "d"
diff --git a/Engine/Source/Runtime/SymsLib/SymsLib.Build.cs b/Engine/Source/Runtime/SymsLib/SymsLib.Build.cs
--- a/Engine/Source/Runtime/SymsLib/SymsLib.Build.cs
+++ b/Engine/Source/Runtime/SymsLib/SymsLib.Build.cs
@@ -26,6 +26,14 @@ public class SymsLib : ModuleRules
 			string LibName = bUseDebug ? "libsymsd_fPIC.a" : "libsyms_fPIC.a";
 			string LibPath = Path.Combine(LibPathBase, "Unix", Target.Architecture.LinuxName);
 
+			PublicAdditionalLibraries.Add(Path.Combine(LibPath, LibName));
+		}
+		else if (Target.Platform.IsInGroup(UnrealPlatformGroup.Apple))
+		{
+			bool bUseDebug = Target.Configuration == UnrealTargetConfiguration.Debug;
+			string LibName = bUseDebug ? "libsymsd.a" : "libsyms.a";
+			string LibPath = Path.Combine(LibPathBase, "Mac", Target.Architecture.AppleName);
+
 			PublicAdditionalLibraries.Add(Path.Combine(LibPath, LibName));
 		}
 	}
