Correct libraries for arm64

diff --git a/Engine/Source/Runtime/SymsLib/BuildForMac.sh b/Engine/Source/Runtime/SymsLib/BuildForMac.sh
--- a/Engine/Source/Runtime/SymsLib/BuildForMac.sh
+++ b/Engine/Source/Runtime/SymsLib/BuildForMac.sh
@@ -33,6 +33,7 @@ BuildLibSyms()
       -DCMAKE_BUILD_TYPE=${FLAVOR} \
       -DCMAKE_POSITION_INDEPENDENT_CODE=ON \
       -DCMAKE_OSX_DEPLOYMENT_TARGET=11.0 \
+      -DCMAKE_OSX_ARCHITECTURES=${ARCH} \
       ${SYMS_DIR}/
     set +x
     
