- Fixed libsyms to be a universal lib, and re-enable it on Mac

diff --git a/Engine/Source/Developer/TraceServices/Private/Model/SymslibResolver.h b/Engine/Source/Developer/TraceServices/Private/Model/SymslibResolver.h
--- a/Engine/Source/Developer/TraceServices/Private/Model/SymslibResolver.h
+++ b/Engine/Source/Developer/TraceServices/Private/Model/SymslibResolver.h
@@ -4,7 +4,7 @@
 #include "HAL/Platform.h"
 
 // We currently only have symslib available on selected platforms.
-#if PLATFORM_WINDOWS || PLATFORM_LINUX
+#if PLATFORM_DESKTOP
 	#define UE_SYMSLIB_AVAILABLE 1
 #else
 	#define UE_SYMSLIB_AVAILABLE 0
diff --git a/Engine/Source/Runtime/SymsLib/SymsLib.Build.cs b/Engine/Source/Runtime/SymsLib/SymsLib.Build.cs
--- a/Engine/Source/Runtime/SymsLib/SymsLib.Build.cs
+++ b/Engine/Source/Runtime/SymsLib/SymsLib.Build.cs
@@ -28,11 +28,11 @@ public class SymsLib : ModuleRules
 
 			PublicAdditionalLibraries.Add(Path.Combine(LibPath, LibName));
 		}
-		else if (Target.Platform.IsInGroup(UnrealPlatformGroup.Apple))
+		else if (Target.Platform == UnrealTargetPlatform.Mac)
 		{
 			bool bUseDebug = Target.Configuration == UnrealTargetConfiguration.Debug;
 			string LibName = bUseDebug ? "libsymsd.a" : "libsyms.a";
-			string LibPath = Path.Combine(LibPathBase, "Mac", Target.Architecture.AppleName);
+			string LibPath = Path.Combine(LibPathBase, "Mac");
 
 			PublicAdditionalLibraries.Add(Path.Combine(LibPath, LibName));
 		}
