Upgrade SymsLib to latest version

diff --git a/Engine/Source/Runtime/SymsLib/BuildForLinux.sh b/Engine/Source/Runtime/SymsLib/BuildForLinux.sh
--- a/Engine/Source/Runtime/SymsLib/BuildForLinux.sh
+++ b/Engine/Source/Runtime/SymsLib/BuildForLinux.sh
@@ -47,7 +47,7 @@ BuildLibSyms()
 
 ( cat <<_EOF_
   ## autogenerated by ${BASH_SOURCE} script
-  SET(LINUX_MULTIARCH_ROOT \$ENV{UE_SDKS_ROOT}/HostLinux/Linux_x64/v19_clang-11.0.1-centos7)
+  SET(LINUX_MULTIARCH_ROOT \$ENV{UE_SDKS_ROOT}/HostLinux/Linux_x64/v22_clang-16.0.6-centos7)
   SET(ARCHITECTURE_TRIPLE \$ENV{ARCH})
 
   message (STATUS "LINUX_MULTIARCH_ROOT is '\${LINUX_MULTIARCH_ROOT}'")
diff --git a/Engine/Source/Runtime/SymsLib/SymsLib.Build.cs b/Engine/Source/Runtime/SymsLib/SymsLib.Build.cs
--- a/Engine/Source/Runtime/SymsLib/SymsLib.Build.cs
+++ b/Engine/Source/Runtime/SymsLib/SymsLib.Build.cs
@@ -32,7 +32,7 @@ public class SymsLib : ModuleRules
 		{
 			bool bUseDebug = Target.Configuration == UnrealTargetConfiguration.Debug;
 			string LibName = bUseDebug ? "libsymsd.a" : "libsyms.a";
-			string LibPath = Path.Combine(LibPathBase, "Mac");
+			string LibPath = Path.Combine(LibPathBase, "Mac", Target.Architecture.bIsX64 ? "x86_64" : "arm64");
 
 			PublicAdditionalLibraries.Add(Path.Combine(LibPath, LibName));
 		}
diff --git a/Engine/Source/Runtime/SymsLib/syms/core/base/syms_base.c b/Engine/Source/Runtime/SymsLib/syms/core/base/syms_base.c
--- a/Engine/Source/Runtime/SymsLib/syms/core/base/syms_base.c
+++ b/Engine/Source/Runtime/SymsLib/syms/core/base/syms_base.c
@@ -7,8 +7,7 @@
 //~ rjf: Library Metadata
 
 SYMS_API SYMS_String8
-syms_version_string(void)
-{
+syms_version_string(void){
   return syms_str8_lit(SYMS_VERSION_STR);
 }
 
@@ -22,8 +21,7 @@ syms_make_u64_range(SYMS_U64 min, SYMS_U64 max){
 }
 
 SYMS_API SYMS_U64Range
-syms_make_u64_inrange(SYMS_U64Range range, SYMS_U64 offset, SYMS_U64 size)
-{
+syms_make_u64_inrange(SYMS_U64Range range, SYMS_U64 offset, SYMS_U64 size){
   SYMS_ASSERT(range.min + offset + size <= range.max);
   SYMS_U64Range result = syms_make_u64_range(range.min + offset, range.min + offset + size);
   return result;
@@ -460,8 +458,12 @@ syms_u64_from_string(SYMS_String8 str, SYMS_U32 radix){
     for (; ptr < opl; ptr += 1){
       SYMS_U8 x = *ptr;
       SYMS_B8 mask = m_from_c[x >> 3];
+      if (!mask){
+        result = 0;
+        break;
+      }
       SYMS_U8 value = (v_from_c[(x - 0x30)%0x20]);
-      if (!mask || value > radix){
+      if (value >= radix){
         result = 0;
         break;
       }
@@ -598,6 +600,63 @@ syms_u64_range_array_from_list(SYMS_Arena *arena, SYMS_U64RangeList *list){
   return(result);
 }
 
+////////////////////////////////
+//~ nick: U64 List Functions
+
+SYMS_API void
+syms_u64_list_push_node(SYMS_U64Node *node, SYMS_U64List *list, SYMS_U64 v){
+  SYMS_QueuePush(list->first, list->last, node);
+  list->count += 1;
+  node->u64 = v;
+}
+
+SYMS_API void
+syms_u64_list_push(SYMS_Arena *arena, SYMS_U64List *list, SYMS_U64 v){
+  SYMS_U64Node *node = syms_push_array(arena, SYMS_U64Node, 1);
+  syms_u64_list_push_node(node, list, v);
+}
+
+SYMS_API void
+syms_u64_list_concat_in_place(SYMS_U64List *dst, SYMS_U64List *src){
+  if (dst->last == 0){
+    *dst = *src;
+  }
+  else if (src->first != 0){
+    dst->last->next = src->first;
+    dst->last = src->last;
+    dst->count += src->count;
+  }
+  syms_memzero_struct(src);
+}
+
+SYMS_API SYMS_U64Array
+syms_u64_array_from_list(SYMS_Arena *arena, SYMS_U64List *list){
+  SYMS_U64Array result = {0};
+  result.u64 = syms_push_array(arena, SYMS_U64, list->count);
+  result.count = list->count;
+  SYMS_U64 *ptr = result.u64;
+  for (SYMS_U64Node *node = list->first;
+       node != 0;
+       node = node->next, ptr += 1){
+    *ptr = node->u64;
+  }
+  return(result);
+}
+
+
+////////////////////////////////
+//~ allen: Array Functions
+
+SYMS_API SYMS_U64
+syms_1based_checked_lookup_u64(SYMS_U64 *u64, SYMS_U64 count, SYMS_U64 n){
+  SYMS_U64 result = 0;
+  if (1 <= n && n <= count){
+    result = u64[n - 1];
+  }
+  return(result);
+}
+
+
 ////////////////////////////////
 //~ allen: Arena Functions
 
diff --git a/Engine/Source/Runtime/SymsLib/syms/core/base/syms_base.h b/Engine/Source/Runtime/SymsLib/syms/core/base/syms_base.h
--- a/Engine/Source/Runtime/SymsLib/syms/core/base/syms_base.h
+++ b/Engine/Source/Runtime/SymsLib/syms/core/base/syms_base.h
@@ -159,6 +159,7 @@ typedef SYMS_U32 SYMS_RegID;
 #define SYMS_Glue(a,b) SYMS_Glue_(a,b)
 
 #define syms_memzero_struct(s) syms_memset((s), 0, sizeof(*(s)))
+#define syms_memisnull_struct(s) syms_memisnull((s), sizeof(*(s)))
 
 #define SYMS_S8_MIN  0x80 // -128
 #define SYMS_S8_MAX  0x7f // +127
@@ -234,6 +235,16 @@ typedef struct SYMS_U64Array{
   SYMS_U64 count;
 } SYMS_U64Array;
 
+typedef struct SYMS_U32Array{
+  SYMS_U32 *u32;
+  SYMS_U64 count;
+} SYMS_U32Array;
+
+typedef struct SYMS_U16Array{
+  SYMS_U16 *u16;
+  SYMS_U64 count;
+} SYMS_U16Array;
+
 typedef struct SYMS_U64ArrayNode{
   struct SYMS_U64ArrayNode *next;
   SYMS_U64 *u64;
@@ -271,6 +282,17 @@ typedef struct SYMS_U64RangeArray{
   SYMS_U64 count;
 } SYMS_U64RangeArray;
 
+typedef struct SYMS_U64Node{
+  struct SYMS_U64Node *next;
+  SYMS_U64 u64;
+} SYMS_U64Node;
+
+typedef struct SYMS_U64List{
+  SYMS_U64Node *first;
+  SYMS_U64Node *last;
+  SYMS_U64 count;
+} SYMS_U64List;
+
 typedef struct SYMS_String8{
   SYMS_U8 *str;
   SYMS_U64 size;
@@ -503,6 +525,19 @@ SYMS_API void syms_u64_range_list_concat(SYMS_U64RangeList *list, SYMS_U64RangeL
 
 SYMS_API SYMS_U64RangeArray syms_u64_range_array_from_list(SYMS_Arena *arena, SYMS_U64RangeList *list);
 
+////////////////////////////////
+//~ nick: U64 List Functions
+
+SYMS_API void syms_u64_list_push_node(SYMS_U64Node *node, SYMS_U64List *list, SYMS_U64 v);
+SYMS_API void syms_u64_list_push(SYMS_Arena *arena, SYMS_U64List *list, SYMS_U64 v);
+SYMS_API void syms_u64_list_concat_in_place(SYMS_U64List *dst, SYMS_U64List *src);
+SYMS_API SYMS_U64Array syms_u64_array_from_list(SYMS_Arena *arena, SYMS_U64List *list);
+
+////////////////////////////////
+//~ allen: Array Functions
+
+SYMS_API SYMS_U64 syms_1based_checked_lookup_u64(SYMS_U64 *u64, SYMS_U64 count, SYMS_U64 n);
+
 ////////////////////////////////
 //~ rjf: Memory/Arena Functions
 
diff --git a/Engine/Source/Runtime/SymsLib/syms/core/base/syms_base_overrides_check.h b/Engine/Source/Runtime/SymsLib/syms/core/base/syms_base_overrides_check.h
--- a/Engine/Source/Runtime/SymsLib/syms/core/base/syms_base_overrides_check.h
+++ b/Engine/Source/Runtime/SymsLib/syms/core/base/syms_base_overrides_check.h
@@ -73,6 +73,10 @@
 # error "SYMS: syms_strlen was not #define'd. It should be #define'd to the name of a strlen implementation."
 #endif
 
+#if !defined(syms_memisnull)
+# error "SYMS: syms_memisnull was not #define'd. It should be #define'd to the name of a memisnull implementation."
+#endif
+
 
 // NOTE(allen): assert
 
diff --git a/Engine/Source/Runtime/SymsLib/syms/core/base/syms_dev.c b/Engine/Source/Runtime/SymsLib/syms/core/base/syms_dev.c
--- a/Engine/Source/Runtime/SymsLib/syms/core/base/syms_dev.c
+++ b/Engine/Source/Runtime/SymsLib/syms/core/base/syms_dev.c
@@ -514,7 +514,7 @@ syms_prof_stringize_basic__dev(SYMS_Arena *arena, SYMS_ProfLock lock, SYMS_Strin
 #define syms_prof_push__dev(...)
 #define syms_prof_paste__dev(...)
 
-#define syms_prof_lock__dev(...) {0}
+#define syms_prof_lock__dev(...) {}
 #define syms_prof_clear__dev(...)
 #define syms_prof_unlock__dev(l) ((void)l)
 
diff --git a/Engine/Source/Runtime/SymsLib/syms/core/base/syms_dev.h b/Engine/Source/Runtime/SymsLib/syms/core/base/syms_dev.h
--- a/Engine/Source/Runtime/SymsLib/syms/core/base/syms_dev.h
+++ b/Engine/Source/Runtime/SymsLib/syms/core/base/syms_dev.h
@@ -25,6 +25,7 @@ enum{
   SYMS_LogFeature_DwarfCFIDecode  = (1 << 5),
   SYMS_LogFeature_DwarfCFIApply   = (1 << 6),
   SYMS_LogFeature_PEEpilog        = (1 << 7),
+  SYMS_LogFeature_PeResParser     = (1 << 8),
   
   
   // Dummy flag shouldn't be associated to anything.
@@ -189,18 +190,30 @@ SYMS_API void syms_prof_stringize_basic__dev(SYMS_Arena *arena, SYMS_ProfLock lo
 
 #if SYMS_ENABLE_DEV_PROFILE
 
-# define SYMS_ProfBegin(str) do{ SYMS_U64 t = SYMS_PROF_TIME(); SYMS_U64 *v = syms_prof_push__dev(); \
+# if !defined(SYMS_ProfBegin)
+#  define SYMS_ProfBegin(str) do{ SYMS_U64 t = SYMS_PROF_TIME(); SYMS_U64 *v = syms_prof_push__dev(); \
 v[0] = (SYMS_U64)(SYMS_THIS_SRCLOC ": " str); v[1] = t; }while(0)
+# endif
 
-# define SYMS_ProfEnd() do{ SYMS_U64 *v = syms_prof_push__dev(); \
+# if !defined(SYMS_ProfEnd)
+#  define SYMS_ProfEnd() do{ SYMS_U64 *v = syms_prof_push__dev(); \
 v[0] = 0; v[1] = SYMS_PROF_TIME(); }while(0)
+# endif
 
-# define SYMS_ProfPasteSubChain(s) syms_prof_paste__dev(s)
+# if !defined(SYMS_ProfPasteSubChain)
+#  define SYMS_ProfPasteSubChain(s) syms_prof_paste__dev(s)
+# endif
 
 #else
-# define SYMS_ProfBegin(str) ((void)0)
-# define SYMS_ProfEnd() ((void)0)
-# define SYMS_ProfPasteSubChain(s) ((void)0)
+# if !defined(SYMS_ProfBegin)
+#  define SYMS_ProfBegin(str) ((void)0)
+# endif
+# if !defined(SYMS_ProfEnd)
+#  define SYMS_ProfEnd() ((void)0)
+# endif
+# if !defined(SYMS_ProfPasteSubChain)
+#  define SYMS_ProfPasteSubChain(s) ((void)0)
+# endif
 #endif
 
 
diff --git a/Engine/Source/Runtime/SymsLib/syms/core/data_structures/syms_data_structures.c b/Engine/Source/Runtime/SymsLib/syms/core/data_structures/syms_data_structures.c
--- a/Engine/Source/Runtime/SymsLib/syms/core/data_structures/syms_data_structures.c
+++ b/Engine/Source/Runtime/SymsLib/syms/core/data_structures/syms_data_structures.c
@@ -4,23 +4,15 @@
 #define SYMS_DATA_STRUCTURES_C
 
 ////////////////////////////////
-//~ NOTE(allen): Syms String Cons
-
-SYMS_API SYMS_StringCons
-syms_string_cons_alloc(SYMS_Arena *arena, SYMS_U64 bucket_count){
-  SYMS_StringCons result = {0};
-  result.bucket_count = bucket_count;
-  result.buckets = syms_push_array_zero(arena, SYMS_StringConsNode*, bucket_count);
-  return(result);
-}
+//~ allen: Syms String Cons
 
 SYMS_API SYMS_String8
 syms_string_cons(SYMS_Arena *arena, SYMS_StringCons *cons, SYMS_String8 string){
   SYMS_ProfBegin("syms_string_cons");
   SYMS_String8 result = {0};
-  if (cons->bucket_count > 0 && string.size > 0){
+  if (SYMS_STRING_CONS_BUCKET_COUNT > 0 && string.size > 0){
     SYMS_U64 hash = syms_hash_djb2(string);
-    SYMS_U64 bucket_index = hash%cons->bucket_count;
+    SYMS_U64 bucket_index = hash%SYMS_STRING_CONS_BUCKET_COUNT;
     for (SYMS_StringConsNode *node = cons->buckets[bucket_index];
          node != 0;
          node = node->next){
@@ -91,7 +83,7 @@ syms_data_idx_cons(SYMS_Arena *arena, SYMS_DataIdxCons *cons, SYMS_String8 data)
 
 
 ////////////////////////////////
-//~ NOTE(allen): U64 Set
+//~ allen: U64 Set
 
 SYMS_API SYMS_U64Set
 syms_u64_set_alloc(SYMS_Arena *arena, SYMS_U64 cap){
@@ -173,7 +165,7 @@ syms_u64_set_erase(SYMS_U64Set *set, SYMS_U64 x){
 
 
 ////////////////////////////////
-//~ NOTE(allen): Syms Spatial Mapping
+//~ allen: Syms Spatial Mapping
 
 //- lookups into spatial maps
 SYMS_API SYMS_U64
@@ -692,7 +684,7 @@ SYMS_API SYMS_String8
 syms_file_id_2_name_map_name_from_id(SYMS_FileID2NameMap *buckets, SYMS_UnitID uid, SYMS_FileID file_id){
   SYMS_ProfBegin("syms_file_id_2_name_map_name_from_id");
   SYMS_U64 hash = syms_file_id_2_name_map_hash(uid, file_id);
-  SYMS_U64 bucket_index = hash%buckets->bucket_count;
+  SYMS_U64 bucket_index = hash%SYMS_FILE_ID_TO_NAME_MAP_BUCKET_COUNT;
   SYMS_String8 result = {0};
   for (SYMS_FileID2NameNode *node = buckets->buckets[bucket_index];
        node != 0;
@@ -712,10 +704,11 @@ syms_file_id_2_name_map_name_from_id(SYMS_FileID2NameMap *buckets, SYMS_UnitID u
 //- copying file id buckets
 SYMS_API SYMS_FileID2NameMap
 syms_file_id_2_name_map_copy(SYMS_Arena *arena, SYMS_StringCons *cons, SYMS_FileID2NameMap *map){
+  SYMS_FileID2NameMap result = {0};
+  
   //- allocate
-  SYMS_U64 bucket_count = map->bucket_count;
   SYMS_U64 count = map->count;
-  SYMS_FileID2NameNode **dst_buckets = syms_push_array(arena, SYMS_FileID2NameNode*, bucket_count);
+  SYMS_FileID2NameNode **dst_buckets = result.buckets;
   SYMS_FileID2NameNode *nodes = syms_push_array(arena, SYMS_FileID2NameNode, count);
   
   //- fill memory
@@ -723,7 +716,7 @@ syms_file_id_2_name_map_copy(SYMS_Arena *arena, SYMS_StringCons *cons, SYMS_File
     SYMS_FileID2NameNode *dst_node = nodes;
     SYMS_FileID2NameNode **dst_bucket = dst_buckets;
     SYMS_FileID2NameNode **src_bucket = map->buckets;
-    SYMS_FileID2NameNode **opl = src_bucket + bucket_count;
+    SYMS_FileID2NameNode **opl = src_bucket + SYMS_FILE_ID_TO_NAME_MAP_BUCKET_COUNT;
     for (; src_bucket < opl; dst_bucket += 1, src_bucket += 1){
       SYMS_FileID2NameNode **chain_ptr = dst_bucket;
       for (SYMS_FileID2NameNode *node = *src_bucket;
@@ -750,9 +743,6 @@ syms_file_id_2_name_map_copy(SYMS_Arena *arena, SYMS_StringCons *cons, SYMS_File
   }
   
   //- fill result
-  SYMS_FileID2NameMap result = {0};
-  result.buckets = dst_buckets;
-  result.bucket_count = bucket_count;
   result.count = count;
   
   return(result);
@@ -760,35 +750,25 @@ syms_file_id_2_name_map_copy(SYMS_Arena *arena, SYMS_StringCons *cons, SYMS_File
 
 //- constructing file id buckets
 
-SYMS_API SYMS_FileID2NameMap
-syms_file_id_2_name_map_alloc(SYMS_Arena *arena, SYMS_U64 bucket_count){
-  SYMS_FileID2NameMap result = {0};
-  result.bucket_count = bucket_count;
-  result.buckets = syms_push_array_zero(arena, SYMS_FileID2NameNode*, bucket_count);
-  return(result);
-}
-
 SYMS_API void
 syms_file_id_2_name_map_insert(SYMS_Arena *arena, SYMS_FileID2NameMap *map,
                                SYMS_UnitID uid, SYMS_FileID file_id, SYMS_String8 name){
   SYMS_ProfBegin("syms_file_id_2_name_map_insert");
-  if (map->bucket_count > 0){
-    SYMS_U64 hash = syms_file_id_2_name_map_hash(uid, file_id);
-    SYMS_U64 bucket_index = hash%map->bucket_count;
-    
-    SYMS_FileID2NameNode *new_node = syms_push_array(arena, SYMS_FileID2NameNode, 1);
-    new_node->uid = uid;
-    new_node->file_id = file_id;
-    new_node->name = name;
-    SYMS_StackPush(map->buckets[bucket_index], new_node);
-    map->count += 1;
-  }
+  SYMS_U64 hash = syms_file_id_2_name_map_hash(uid, file_id);
+  SYMS_U64 bucket_index = hash%SYMS_FILE_ID_TO_NAME_MAP_BUCKET_COUNT;
+  
+  SYMS_FileID2NameNode *new_node = syms_push_array(arena, SYMS_FileID2NameNode, 1);
+  new_node->uid = uid;
+  new_node->file_id = file_id;
+  new_node->name = name;
+  SYMS_StackPush(map->buckets[bucket_index], new_node);
+  map->count += 1;
   SYMS_ProfEnd();
 }
 
 
 ////////////////////////////////
-//~ NOTE(allen): Mapping Functions (String -> {UnitID,FileID})
+//~ allen: Mapping Functions (String -> {UnitID,FileID})
 
 //- copying file maps
 
@@ -909,7 +889,7 @@ syms_name_2_file_id_map_loose_push(SYMS_Arena *arena, SYMS_Name2FileIDMapLoose *
 
 
 ////////////////////////////////
-//~ NOTE(allen): ID Mapping Functions
+//~ allen: ID Mapping Functions
 
 //- copying id maps
 
@@ -1016,7 +996,7 @@ syms_id_map_insert(SYMS_Arena *arena, SYMS_IDMap *map, SYMS_U64 key, void *val){
 
 
 ////////////////////////////////
-//~ NOTE(allen): Symbol Name Mapping Structure (String -> Array(USID))
+//~ allen: Symbol Name Mapping Structure (String -> Array(SID))
 
 //- lookups into symbol name map
 
@@ -1024,16 +1004,14 @@ SYMS_API SYMS_SymbolIDArray
 syms_symbol_name_map_array_from_string(SYMS_SymbolNameMap *map, SYMS_String8 string){
   SYMS_ProfBegin("syms_symbol_name_map_array_from_string");
   SYMS_SymbolIDArray result = {0};
-  if (map->bucket_count > 0){
-    SYMS_U64 hash = syms_hash_djb2(string);
-    SYMS_U64 bucket_index = hash%map->bucket_count;
-    for (SYMS_SymbolNameNode *node = map->buckets[bucket_index];
-         node != 0;
-         node = node->next_bucket){
-      if (node->hash == hash && syms_string_match(node->name, string, 0)){
-        result = node->sid_array;
-        break;
-      }
+  SYMS_U64 hash = syms_hash_djb2(string);
+  SYMS_U64 bucket_index = hash%SYMS_SYMBOL_NAME_MAP_BUCKET_COUNT;
+  for (SYMS_SymbolNameNode *node = map->buckets[bucket_index];
+       node != 0;
+       node = node->next_bucket){
+    if (node->hash == hash && syms_string_match(node->name, string, 0)){
+      result = node->sid_array;
+      break;
     }
   }
   SYMS_ProfEnd();
@@ -1042,52 +1020,44 @@ syms_symbol_name_map_array_from_string(SYMS_SymbolNameMap *map, SYMS_String8 str
 
 //- constructing symbol name maps
 
-SYMS_API SYMS_SymbolNameMapLoose
-syms_symbol_name_map_begin(SYMS_Arena *arena, SYMS_U64 bucket_count){
-  SYMS_SymbolNameMapLoose result = {0};
-  result.bucket_count = bucket_count;
-  result.buckets = syms_push_array_zero(arena, SYMS_SymbolNameNodeLoose*, bucket_count);
-  return(result);
-}
-
 SYMS_API void
 syms_symbol_name_map_push(SYMS_Arena *arena, SYMS_SymbolNameMapLoose *map,
                           SYMS_String8 name, SYMS_SymbolID sid){
   SYMS_ProfBegin("syms_symbol_name_map_push");
-  if (map->bucket_count > 0){
-    SYMS_U64 hash = syms_hash_djb2(name);
-    SYMS_U64 bucket_index = hash%map->bucket_count;
-    
-    // find existing node
-    SYMS_SymbolNameNodeLoose *node = 0;
-    for (SYMS_SymbolNameNodeLoose *bucket = map->buckets[bucket_index];
-         bucket != 0;
-         bucket = bucket->next_bucket){
-      if (bucket->hash == hash &&
-          syms_string_match(name, bucket->name, 0)){
-        node = bucket;
-        break;
-      }
-    }
-    
-    // create node
-    if (node == 0){
-      node = syms_push_array_zero(arena, SYMS_SymbolNameNodeLoose, 1);
-      SYMS_StackPush_N(map->buckets[bucket_index], node, next_bucket);
-      SYMS_QueuePush(map->first, map->last, node);
-      map->node_count += 1;
-      node->name = syms_push_string_copy(arena, name);
-      node->hash = hash;
-    }
-    
-    // insert sid
-    {
-      SYMS_SymbolIDNode *sid_node = syms_push_array_zero(arena, SYMS_SymbolIDNode, 1);
-      SYMS_QueuePush(node->sid_list.first, node->sid_list.last, sid_node);
-      node->sid_list.count += 1;
-      sid_node->id = sid;
+  
+  SYMS_U64 hash = syms_hash_djb2(name);
+  SYMS_U64 bucket_index = hash%SYMS_SYMBOL_NAME_MAP_BUCKET_COUNT;
+  
+  // find existing node
+  SYMS_SymbolNameNodeLoose *node = 0;
+  for (SYMS_SymbolNameNodeLoose *bucket = map->buckets[bucket_index];
+       bucket != 0;
+       bucket = bucket->next_bucket){
+    if (bucket->hash == hash &&
+        syms_string_match(name, bucket->name, 0)){
+      node = bucket;
+      break;
     }
   }
+  
+  // create node
+  if (node == 0){
+    node = syms_push_array_zero(arena, SYMS_SymbolNameNodeLoose, 1);
+    SYMS_StackPush_N(map->buckets[bucket_index], node, next_bucket);
+    SYMS_QueuePush(map->first, map->last, node);
+    map->node_count += 1;
+    node->name = syms_push_string_copy(arena, name);
+    node->hash = hash;
+  }
+  
+  // insert sid
+  {
+    SYMS_SymbolIDNode *sid_node = syms_push_array_zero(arena, SYMS_SymbolIDNode, 1);
+    SYMS_QueuePush(node->sid_list.first, node->sid_list.last, sid_node);
+    node->sid_list.count += 1;
+    sid_node->id = sid;
+  }
+  
   SYMS_ProfEnd();
 }
 
@@ -1095,13 +1065,13 @@ SYMS_API SYMS_SymbolNameMap
 syms_symbol_name_map_bake(SYMS_Arena *arena, SYMS_SymbolNameMapLoose *loose){
   SYMS_ProfBegin("syms_symbol_name_map_bake");
   
+  SYMS_SymbolNameMap result = {0};
+  SYMS_SymbolNameNode **buckets = result.buckets;
+  
   //- allocate map memory
   SYMS_U64 node_count = loose->node_count;
   SYMS_SymbolNameNode *nodes = syms_push_array(arena, SYMS_SymbolNameNode, node_count);
   
-  SYMS_U64 bucket_count = node_count*3/2;
-  SYMS_SymbolNameNode **buckets = syms_push_array_zero(arena, SYMS_SymbolNameNode*, bucket_count);
-  
   //- fill map from loose
   SYMS_SymbolNameNode *node_ptr = nodes;
   for (SYMS_SymbolNameNodeLoose *loose_node = loose->first;
@@ -1128,7 +1098,7 @@ syms_symbol_name_map_bake(SYMS_Arena *arena, SYMS_SymbolNameMapLoose *loose){
     }
     
     // insert node to bucket
-    SYMS_U64 bucket_index = hash%bucket_count;
+    SYMS_U64 bucket_index = hash%SYMS_SYMBOL_NAME_MAP_BUCKET_COUNT;
     SYMS_StackPush_N(buckets[bucket_index], node_ptr, next_bucket);
     
     // fill node ptr
@@ -1139,9 +1109,6 @@ syms_symbol_name_map_bake(SYMS_Arena *arena, SYMS_SymbolNameMapLoose *loose){
   }
   
   //- assemble result
-  SYMS_SymbolNameMap result = {0};
-  result.buckets = buckets;
-  result.bucket_count = bucket_count;
   result.nodes = nodes;
   result.node_count = node_count;
   
@@ -1151,7 +1118,7 @@ syms_symbol_name_map_bake(SYMS_Arena *arena, SYMS_SymbolNameMapLoose *loose){
 
 
 ////////////////////////////////
-//~ NOTE(allen): Line Tables
+//~ allen: Line Tables
 
 //- lookups into line tables
 
@@ -1278,7 +1245,287 @@ syms_line_table_with_indexes_from_parse(SYMS_Arena *arena, SYMS_LineParseOut *pa
 
 
 ////////////////////////////////
-//~ NOTE(allen): Copies & Operators for Other Data Structures
+//~ allen: Line To Addr Map
+
+SYMS_API SYMS_FileToLineToAddrMap
+syms_line_to_addr_map_from_line_table(SYMS_Arena *arena, SYMS_LineTable *line_table){
+  SYMS_ArenaTemp scratch = syms_get_scratch(0, 0);
+  
+  //- setup loose map
+  SYMS_FileToLineToAddrLoose loose = {0};
+  
+  //- last-used cache for file nodes
+  SYMS_FileToLineToAddrLooseFile *cached_file_node = 0;
+  SYMS_U64 cached_file_id = 0;
+  
+  //- read lines
+  SYMS_U64 *seq_idx_ptr = line_table->sequence_index_array;
+  SYMS_U64 seq_count = line_table->sequence_count;
+  for (SYMS_U64 i = 0; i < seq_count; i += 1){
+    // get sequence range & inc
+    SYMS_U64 first = *seq_idx_ptr;
+    seq_idx_ptr += 1;
+    SYMS_U64 opl = *seq_idx_ptr;
+    
+    // iterate lines
+    if (first < opl){
+      SYMS_Line *line_ptr = line_table->line_array + first;
+      SYMS_U64 last = opl - 1;
+      for (SYMS_U64 j = first; j < last; j += 1, line_ptr += 1){
+        // grab line data
+        SYMS_U64 line_file_id = line_ptr->src_coord.file_id;
+        SYMS_U32 line_number = line_ptr->src_coord.line;
+        
+        // get file node
+        SYMS_FileToLineToAddrLooseFile *file_node = 0;
+        if (line_file_id == cached_file_id){
+          file_node = cached_file_node;
+        }
+        else{
+          for (SYMS_FileToLineToAddrLooseFile *node = loose.first;
+               node != 0;
+               node = node->next){
+            if (node->file_id == line_file_id){
+              file_node = node;
+              break;
+            }
+          }
+        }
+        if (file_node == 0){
+          file_node = syms_push_array_zero(scratch.arena, SYMS_FileToLineToAddrLooseFile, 1);
+          SYMS_QueuePush(loose.first, loose.last, file_node);
+          loose.count += 1;
+          file_node->file_id = line_file_id;
+        }
+        
+        // update the file node cache slot
+        cached_file_id = line_file_id;
+        cached_file_node = file_node;
+        
+        // get line node
+        SYMS_FileToLineToAddrLooseLine *line_node = 0;
+        for (SYMS_FileToLineToAddrLooseLine *node = file_node->first;
+             node != 0;
+             node = node->next){
+          if (node->line == line_number){
+            line_node = node;
+            break;
+          }
+        }
+        if (line_node == 0){
+          line_node = syms_push_array_zero(scratch.arena, SYMS_FileToLineToAddrLooseLine, 1);
+          SYMS_QueuePush(file_node->first, file_node->last, line_node);
+          file_node->line_count += 1;
+          line_node->line = line_number;
+        }
+        
+        // push range
+        SYMS_U64Range range = syms_make_u64_range(line_ptr->voff, (line_ptr + 1)->voff);
+        syms_u64_range_list_push(scratch.arena, &line_node->ranges, range);
+        file_node->range_count += 1;
+      }
+    }
+  }
+  
+  //- convert loose to buckets & maps
+  SYMS_U64 bucket_count = 0;
+  if (loose.count > 0){
+    bucket_count = ((loose.count + 1)*3/2) | 1;
+  }
+  
+  SYMS_FileToLineToAddrNode **buckets = syms_push_array_zero(arena, SYMS_FileToLineToAddrNode*, bucket_count);
+  for (SYMS_FileToLineToAddrLooseFile *loose_file_node = loose.first;
+       loose_file_node != 0;
+       loose_file_node = loose_file_node->next){
+    SYMS_ArenaTemp temp = syms_arena_temp_begin(scratch.arena);
+    
+    // grab counts
+    SYMS_U64 line_count = loose_file_node->line_count;
+    SYMS_U64 range_count = loose_file_node->range_count;
+    
+    // create sorted node array
+    SYMS_FileToLineToAddrLooseLine **array = syms_push_array(scratch.arena, SYMS_FileToLineToAddrLooseLine*,
+                                                             line_count);
+    {
+      SYMS_FileToLineToAddrLooseLine **line_ptr = array;
+      for (SYMS_FileToLineToAddrLooseLine *loose_line_node = loose_file_node->first;
+           loose_line_node != 0;
+           loose_line_node = loose_line_node->next, line_ptr += 1){
+        *line_ptr = loose_line_node;
+      }
+    }
+    syms_line_to_addr_line_sort(array, line_count);
+    
+    // fill line map arrays
+    SYMS_U64Range *ranges = syms_push_array(arena, SYMS_U64Range, range_count);
+    SYMS_U32 *line_range_indexes = syms_push_array(arena, SYMS_U32, line_count + 1);
+    SYMS_U32 *line_numbers = syms_push_array(arena, SYMS_U32, line_count);
+    
+    SYMS_U32 line_range_index = 0;
+    SYMS_U32 *line_range_index_ptr = line_range_indexes;
+    SYMS_U32 *line_number_ptr = line_numbers;
+    
+    {
+      SYMS_FileToLineToAddrLooseLine **line_ptr = array;
+      for (SYMS_U64 i = 0; i < line_count; i += 1, line_ptr += 1){
+        // fills
+        *line_number_ptr = (**line_ptr).line;
+        *line_range_index_ptr = line_range_index;
+        SYMS_U64Range *first_range_ptr = ranges + line_range_index;
+        SYMS_U64Range *range_ptr = first_range_ptr;
+        for (SYMS_U64RangeNode *node = (**line_ptr).ranges.first;
+             node != 0;
+             node = node->next, range_ptr += 1){
+          *range_ptr = node->range;
+        }
+        
+        // incs
+        line_number_ptr += 1;
+        line_range_index_ptr += 1;
+        line_range_index += (SYMS_U64)(range_ptr - first_range_ptr);
+      }
+      
+      // fill ender index
+      *line_range_index_ptr = line_range_index;
+    }
+    
+    // assemble the line map
+    SYMS_LineToAddrMap *new_map = syms_push_array(arena, SYMS_LineToAddrMap, 1);
+    new_map->ranges = ranges;
+    new_map->line_range_indexes = line_range_indexes;
+    new_map->line_numbers = line_numbers;
+    new_map->line_count = line_count;
+    
+    // insert bucket
+    SYMS_FileToLineToAddrNode *new_file_node = syms_push_array(arena, SYMS_FileToLineToAddrNode, 1);
+    new_file_node->file_id = loose_file_node->file_id;
+    new_file_node->map = new_map;
+    SYMS_U64 bucket_index = loose_file_node->file_id%bucket_count;
+    SYMS_StackPush(buckets[bucket_index], new_file_node);
+    
+    syms_arena_temp_end(temp);
+  }
+  
+  syms_release_scratch(scratch);
+  
+  //- build result
+  SYMS_FileToLineToAddrMap result = {0};
+  result.buckets = buckets;
+  result.bucket_count = bucket_count;
+  
+  return(result);
+}
+
+SYMS_API SYMS_LineToAddrMap*
+syms_line_to_addr_map_lookup_file_id(SYMS_FileToLineToAddrMap *map, SYMS_FileID file_id){
+  SYMS_LineToAddrMap *result = &syms_line_to_addr_map_nil;
+  if (map->bucket_count > 0){
+    SYMS_U64 bucket_index = file_id%map->bucket_count;
+    for (SYMS_FileToLineToAddrNode *node = map->buckets[bucket_index];
+         node != 0;
+         node = node->next){
+      if (node->file_id == file_id){
+        result = node->map;
+        break;
+      }
+    }
+  }
+  return(result);
+}
+
+SYMS_API SYMS_U64RangeArray
+syms_line_to_addr_map_lookup_nearest_line_number(SYMS_LineToAddrMap *map, SYMS_U32 line,
+                                                 SYMS_U32 *actual_line_out){
+  SYMS_U64 count = map->line_count;
+  SYMS_U32 *numbers = map->line_numbers;
+  SYMS_U64 index = syms_index_from_n__u32__binary_search_round_up(numbers, count, line);
+  if (count > 0 && index >= count){
+    index = count - 1;
+  }
+  SYMS_U64RangeArray result = {0};
+  if (index < count){
+    SYMS_U32 *range_indexes = map->line_range_indexes;
+    result.ranges = map->ranges + range_indexes[index];
+    result.count = range_indexes[index + 1] - range_indexes[index];
+    *actual_line_out = numbers[index];
+  }
+  return(result);
+}
+
+
+SYMS_API void
+syms_line_to_addr_line_sort(SYMS_FileToLineToAddrLooseLine **array, SYMS_U64 count){
+  SYMS_ProfBegin("syms_line_to_addr_line_sort");
+  syms_line_to_addr_line_sort__rec(array, count);
+  SYMS_ProfEnd();
+}
+
+SYMS_API void
+syms_line_to_addr_line_sort__rec(SYMS_FileToLineToAddrLooseLine **array, SYMS_U64 count){
+  if (count > 4){
+    SYMS_U64 last = count - 1;
+    
+    // swap
+    SYMS_U64 mid = count/2;
+    SYMS_Swap(SYMS_FileToLineToAddrLooseLine*, array[mid], array[last]);
+    
+    // partition
+    SYMS_U64 key = array[last]->line;
+    SYMS_U64 j = 0;
+    for (SYMS_U64 i = 0; i < last; i += 1){
+      if (array[i]->line < key){
+        if (j != i){
+          SYMS_Swap(SYMS_FileToLineToAddrLooseLine*, array[i], array[j]);
+        }
+        j += 1;
+      }
+    }
+    
+    SYMS_Swap(SYMS_FileToLineToAddrLooseLine*, array[j], array[last]);
+    
+    // recurse
+    SYMS_U64 pivot = j;
+    syms_line_to_addr_line_sort__rec(array, pivot);
+    syms_line_to_addr_line_sort__rec(array + pivot + 1, (count - pivot - 1));
+  }
+  else if (count == 2){
+    if (array[0]->line > array[1]->line){
+      SYMS_Swap(SYMS_FileToLineToAddrLooseLine*, array[0], array[1]);
+    }
+  }
+  else if (count == 3){
+    if (array[0]->line > array[1]->line){
+      SYMS_Swap(SYMS_FileToLineToAddrLooseLine*, array[0], array[1]);
+    }
+    if (array[1]->line > array[2]->line){
+      SYMS_Swap(SYMS_FileToLineToAddrLooseLine*, array[1], array[2]);
+      if (array[0]->line > array[1]->line){
+        SYMS_Swap(SYMS_FileToLineToAddrLooseLine*, array[0], array[1]);
+      }
+    }
+  }
+  else if (count == 4){
+    if (array[0]->line > array[1]->line){
+      SYMS_Swap(SYMS_FileToLineToAddrLooseLine*, array[0], array[1]);
+    }
+    if (array[2]->line > array[3]->line){
+      SYMS_Swap(SYMS_FileToLineToAddrLooseLine*, array[2], array[3]);
+    }
+    if (array[0]->line > array[2]->line){
+      SYMS_Swap(SYMS_FileToLineToAddrLooseLine*, array[0], array[2]);
+    }
+    if (array[1]->line > array[3]->line){
+      SYMS_Swap(SYMS_FileToLineToAddrLooseLine*, array[1], array[3]);
+    }
+    if (array[1]->line > array[2]->line){
+      SYMS_Swap(SYMS_FileToLineToAddrLooseLine*, array[1], array[2]);
+    }
+  }
+}
+
+
+////////////////////////////////
+//~ allen: Copies & Operators for Other Data Structures
 
 SYMS_API SYMS_String8Array
 syms_string_array_copy(SYMS_Arena *arena, SYMS_StringCons *cons, SYMS_String8Array *array){
@@ -1313,7 +1560,6 @@ syms_string_array_copy(SYMS_Arena *arena, SYMS_StringCons *cons, SYMS_String8Arr
   return(result);
 }
 
-
 SYMS_API SYMS_LinkNameRecArray
 syms_link_name_record_copy(SYMS_Arena *arena, SYMS_LinkNameRecArray *array){
   SYMS_ProfBegin("syms_link_name_record_copy");
@@ -1332,4 +1578,42 @@ syms_link_name_record_copy(SYMS_Arena *arena, SYMS_LinkNameRecArray *array){
   return(result);
 }
 
+////////////////////////////////
+//~ allen: Binary Search Functions
+
+SYMS_API SYMS_U64
+syms_index_from_n__u32__binary_search_round_up(SYMS_U32 *v, SYMS_U64 count, SYMS_U32 n){
+  SYMS_ProfBegin("syms_index_from_n__u32__binary_search_round_up");
+  SYMS_U64 result = SYMS_U64_MAX;
+  if (count > 0 && n <= v[count - 1]){
+    //- binary search:
+    //   minimum index s.t. v[index] >= n
+    //  in this one we assume:
+    //   (i != j) implies (v[i] != v[j])
+    //   thus if (v[index] == n) then index already satisfies the requirement
+    SYMS_U64 first = 0;
+    SYMS_U64 opl = count;
+    for (;;){
+      SYMS_U64 mid = (first + opl - 1)/2;
+      SYMS_U64 w = v[mid];
+      if (w < n){
+        first = mid + 1;
+      }
+      else if (w > n){
+        opl = mid + 1;
+      }
+      else{
+        first = mid;
+        break;
+      }
+      if (first + 1 >= opl){
+        break;
+      }
+    }
+    result = first;
+  }
+  SYMS_ProfEnd();
+  return(result);
+}
+
 #endif //SYMS_DATA_STRUCTURES_C
diff --git a/Engine/Source/Runtime/SymsLib/syms/core/data_structures/syms_data_structures.h b/Engine/Source/Runtime/SymsLib/syms/core/data_structures/syms_data_structures.h
--- a/Engine/Source/Runtime/SymsLib/syms/core/data_structures/syms_data_structures.h
+++ b/Engine/Source/Runtime/SymsLib/syms/core/data_structures/syms_data_structures.h
@@ -9,6 +9,8 @@
 
 // deduplicates equivalent strings
 
+#define SYMS_STRING_CONS_BUCKET_COUNT 1024
+
 typedef struct SYMS_StringConsNode{
   struct SYMS_StringConsNode *next;
   SYMS_String8 string;
@@ -16,8 +18,7 @@ typedef struct SYMS_StringConsNode{
 } SYMS_StringConsNode;
 
 typedef struct SYMS_StringCons{
-  SYMS_StringConsNode **buckets;
-  SYMS_U64 bucket_count;
+  SYMS_StringConsNode *buckets[SYMS_STRING_CONS_BUCKET_COUNT];
 } SYMS_StringCons;
 
 // assign sequential indexes to small variable length blobs
@@ -102,6 +103,8 @@ typedef struct SYMS_1DEndPoint{
 // maps a unit-id,file-id pair to a string
 // organized as a hash table to opimize for key based lookups
 
+#define SYMS_FILE_ID_TO_NAME_MAP_BUCKET_COUNT 1024
+
 typedef struct SYMS_FileID2NameNode{
   struct SYMS_FileID2NameNode *next;
   // key
@@ -112,8 +115,7 @@ typedef struct SYMS_FileID2NameNode{
 } SYMS_FileID2NameNode;
 
 typedef struct SYMS_FileID2NameMap{
-  SYMS_FileID2NameNode **buckets;
-  SYMS_U64 bucket_count;
+  SYMS_FileID2NameNode *buckets[SYMS_FILE_ID_TO_NAME_MAP_BUCKET_COUNT];
   SYMS_U64 count;
 } SYMS_FileID2NameMap;
 
@@ -190,13 +192,15 @@ typedef struct SYMS_IDMap{
 
 
 ////////////////////////////////
-//~ allen: Symbol Name Mapping Structure (String -> Array(USID))
+//~ allen: Symbol Name Mapping Structure (String -> Array(SID))
 
-// maps strings to lists of USIDs
+// maps strings to lists of SIDs
 // organized as an array of nodes and a hash table simultaneously
 // so that single name lookups are accelerated, and filter matching
 // is also possible.
 
+#define SYMS_SYMBOL_NAME_MAP_BUCKET_COUNT 1024
+
 typedef struct SYMS_SymbolNameNode{
   struct SYMS_SymbolNameNode *next_bucket;
   SYMS_String8 name;
@@ -205,8 +209,7 @@ typedef struct SYMS_SymbolNameNode{
 } SYMS_SymbolNameNode;
 
 typedef struct SYMS_SymbolNameMap{
-  SYMS_SymbolNameNode **buckets;
-  SYMS_U64 bucket_count;
+  SYMS_SymbolNameNode *buckets[SYMS_SYMBOL_NAME_MAP_BUCKET_COUNT];
   SYMS_SymbolNameNode *nodes;
   SYMS_U64 node_count;
 } SYMS_SymbolNameMap;
@@ -221,18 +224,64 @@ typedef struct SYMS_SymbolNameNodeLoose{
 } SYMS_SymbolNameNodeLoose;
 
 typedef struct SYMS_SymbolNameMapLoose{
-  SYMS_SymbolNameNodeLoose **buckets;
-  SYMS_U64 bucket_count;
+  SYMS_SymbolNameNodeLoose *buckets[SYMS_SYMBOL_NAME_MAP_BUCKET_COUNT];
   SYMS_SymbolNameNodeLoose *first;
   SYMS_SymbolNameNodeLoose *last;
   SYMS_U64 node_count;
 } SYMS_SymbolNameMapLoose;
 
+
+////////////////////////////////
+//~ allen: Line To Addr Map Structure
+
+typedef struct SYMS_LineToAddrMap{
+  SYMS_U64Range *ranges;
+  // line_range_indexes ranges from [0,line_count] inclusive so that:
+  // for-all i in [0,line_count):
+  //   (line_range_indexes[i + 1] - line_range_indexes[i]) == # of ranges for line at index i
+  SYMS_U32 *line_range_indexes;
+  SYMS_U32 *line_numbers;
+  SYMS_U64 line_count;
+} SYMS_LineToAddrMap;
+
+typedef struct SYMS_FileToLineToAddrNode{
+  struct SYMS_FileToLineToAddrNode *next;
+  SYMS_FileID file_id;
+  SYMS_LineToAddrMap *map;
+} SYMS_FileToLineToAddrNode;
+
+typedef struct SYMS_FileToLineToAddrMap{
+  SYMS_FileToLineToAddrNode **buckets;
+  SYMS_U64 bucket_count;
+} SYMS_FileToLineToAddrMap;
+
+//- loose version
+typedef struct SYMS_FileToLineToAddrLooseLine{
+  struct SYMS_FileToLineToAddrLooseLine *next;
+  SYMS_U32 line;
+  SYMS_U64RangeList ranges;
+} SYMS_FileToLineToAddrLooseLine;
+
+typedef struct SYMS_FileToLineToAddrLooseFile{
+  struct SYMS_FileToLineToAddrLooseFile *next;
+  SYMS_FileID file_id;
+  SYMS_FileToLineToAddrLooseLine *first;
+  SYMS_FileToLineToAddrLooseLine *last;
+  SYMS_U64 line_count;
+  SYMS_U64 range_count;
+} SYMS_FileToLineToAddrLooseFile;
+
+typedef struct SYMS_FileToLineToAddrLoose{
+  SYMS_FileToLineToAddrLooseFile *first;
+  SYMS_FileToLineToAddrLooseFile *last;
+  SYMS_U64 count;
+} SYMS_FileToLineToAddrLoose;
+
+
 ////////////////////////////////
 //~ allen: String Cons Functions
 
-SYMS_API SYMS_StringCons syms_string_cons_alloc(SYMS_Arena *arena, SYMS_U64 bucket_count);
-SYMS_API SYMS_String8    syms_string_cons(SYMS_Arena *arena, SYMS_StringCons *cons, SYMS_String8 string);
+SYMS_API SYMS_String8     syms_string_cons(SYMS_Arena *arena, SYMS_StringCons *cons, SYMS_String8 string);
 
 SYMS_API SYMS_DataIdxCons syms_data_idx_cons_alloc(SYMS_Arena *arena, SYMS_U64 bucket_count);
 SYMS_API SYMS_U64         syms_data_idx_cons(SYMS_Arena *arena, SYMS_DataIdxCons *cons, SYMS_String8 data);
@@ -286,7 +335,7 @@ SYMS_API SYMS_B32          syms_spatial_map_1d_invariants(SYMS_SpatialMap1D *map
 SYMS_API SYMS_U64            syms_file_id_2_name_map_hash(SYMS_UnitID uid, SYMS_FileID file_id);
 
 //- lookups into file id buckets
-SYMS_API SYMS_String8        syms_file_id_2_name_map_name_from_id(SYMS_FileID2NameMap *buckets,
+SYMS_API SYMS_String8        syms_file_id_2_name_map_name_from_id(SYMS_FileID2NameMap *map,
                                                                   SYMS_UnitID uid, SYMS_FileID file_id);
 
 //- copying file id buckets
@@ -294,7 +343,6 @@ SYMS_API SYMS_FileID2NameMap syms_file_id_2_name_map_copy(SYMS_Arena *arena, SYM
                                                           SYMS_FileID2NameMap *map);
 
 //- constructing file id buckets
-SYMS_API SYMS_FileID2NameMap syms_file_id_2_name_map_alloc(SYMS_Arena *arena, SYMS_U64 bucket_count);
 SYMS_API void                syms_file_id_2_name_map_insert(SYMS_Arena *arena, SYMS_FileID2NameMap *map,
                                                             SYMS_UnitID uid, SYMS_FileID file_id,
                                                             SYMS_String8 name);
@@ -330,7 +378,7 @@ SYMS_API void         syms_id_map_insert(SYMS_Arena *arena, SYMS_IDMap *map, SYM
 
 
 ////////////////////////////////
-//~ allen: Symbol Name Mapping Structure (String -> Array(USID))
+//~ allen: Symbol Name Mapping Structure (String -> Array(SID))
 
 // TODO(allen): copying
 
@@ -338,7 +386,6 @@ SYMS_API void         syms_id_map_insert(SYMS_Arena *arena, SYMS_IDMap *map, SYM
 SYMS_API SYMS_SymbolIDArray syms_symbol_name_map_array_from_string(SYMS_SymbolNameMap *map, SYMS_String8 string);
 
 //- constructing symbol name maps
-SYMS_API SYMS_SymbolNameMapLoose syms_symbol_name_map_begin(SYMS_Arena *arena, SYMS_U64 bucket_count);
 SYMS_API void                    syms_symbol_name_map_push(SYMS_Arena *arena, SYMS_SymbolNameMapLoose *map,
                                                            SYMS_String8 name, SYMS_SymbolID sid);
 SYMS_API SYMS_SymbolNameMap      syms_symbol_name_map_bake(SYMS_Arena *arena, SYMS_SymbolNameMapLoose *loose);
@@ -358,6 +405,24 @@ SYMS_API void              syms_line_table_rewrite_file_ids_in_place(SYMS_FileID
                                                                      SYMS_LineTable *line_table_in_out);
 SYMS_API SYMS_LineTable    syms_line_table_with_indexes_from_parse(SYMS_Arena *arena, SYMS_LineParseOut *parse);
 
+////////////////////////////////
+//~ allen: Line To Addr Map
+
+//- line-to-addr map
+SYMS_API SYMS_FileToLineToAddrMap
+syms_line_to_addr_map_from_line_table(SYMS_Arena *arena, SYMS_LineTable *table);
+
+//- line-to-addr query
+SYMS_API SYMS_LineToAddrMap* syms_line_to_addr_map_lookup_file_id(SYMS_FileToLineToAddrMap *map,
+                                                                  SYMS_FileID file_id);
+
+SYMS_API SYMS_U64RangeArray syms_line_to_addr_map_lookup_nearest_line_number(SYMS_LineToAddrMap *map,
+                                                                             SYMS_U32 line,
+                                                                             SYMS_U32 *actual_line_out);
+
+//- line-to-addr map helpers
+SYMS_API void syms_line_to_addr_line_sort(SYMS_FileToLineToAddrLooseLine **array, SYMS_U64 count);
+SYMS_API void syms_line_to_addr_line_sort__rec(SYMS_FileToLineToAddrLooseLine **array, SYMS_U64 count);
 
 ////////////////////////////////
 //~ allen: Copies & Operators for Other Data Structures
@@ -365,7 +430,11 @@ SYMS_API SYMS_LineTable    syms_line_table_with_indexes_from_parse(SYMS_Arena *a
 SYMS_API SYMS_String8Array syms_string_array_copy(SYMS_Arena *arena, SYMS_StringCons *cons_optional,
                                                   SYMS_String8Array *array);
 
-SYMS_API SYMS_LinkNameRecArray  syms_link_name_record_copy(SYMS_Arena *arena, SYMS_LinkNameRecArray *array);
+SYMS_API SYMS_LinkNameRecArray syms_link_name_record_copy(SYMS_Arena *arena, SYMS_LinkNameRecArray *array);
+
+////////////////////////////////
+//~ allen: Binary Search Functions
 
+SYMS_API SYMS_U64 syms_index_from_n__u32__binary_search_round_up(SYMS_U32 *v, SYMS_U64 count, SYMS_U32 n);
 
 #endif //SYMS_DATA_STRUCTURES_H
diff --git a/Engine/Source/Runtime/SymsLib/syms/core/dwarf/syms_dwarf.c b/Engine/Source/Runtime/SymsLib/syms/core/dwarf/syms_dwarf.c
--- a/Engine/Source/Runtime/SymsLib/syms/core/dwarf/syms_dwarf.c
+++ b/Engine/Source/Runtime/SymsLib/syms/core/dwarf/syms_dwarf.c
@@ -333,4 +333,170 @@ syms_dw_type_modifiers_from_tag_kind(SYMS_DwTagKind tag_kind)
   return mods;
 }
 
+SYMS_API SYMS_DwAttribClass
+syms_dw_attrib_class_from_attrib_kind_extensions(SYMS_DwAttribKind attrib)
+{
+  SYMS_DwAttribClass result = 0;
+  switch (attrib) {
+  case SYMS_DwAttribKind_GNU_VECTOR:                     result = SYMS_DwAttribClass_FLAG;      break;
+  case SYMS_DwAttribKind_GNU_GUARDED_BY:                 result = SYMS_DwAttribClass_UNDEFINED; break;
+  case SYMS_DwAttribKind_GNU_PT_GUARDED_BY:              result = SYMS_DwAttribClass_UNDEFINED; break;
+  case SYMS_DwAttribKind_GNU_GUARDED:                    result = SYMS_DwAttribClass_UNDEFINED; break;
+  case SYMS_DwAttribKind_GNU_LOCKS_EXCLUDED:             result = SYMS_DwAttribClass_UNDEFINED; break;
+  case SYMS_DwAttribKind_GNU_EXCLUSIVE_LOCKS_REQUIRED:   result = SYMS_DwAttribClass_UNDEFINED; break;
+  case SYMS_DwAttribKind_GNU_SHARED_LOCKS_REQUIRED:      result = SYMS_DwAttribClass_UNDEFINED; break;
+  case SYMS_DwAttribKind_GNU_ODR_SIGNATURE:              result = SYMS_DwAttribClass_UNDEFINED; break;
+  case SYMS_DwAttribKind_GNU_TEMPLATE_NAME:              result = SYMS_DwAttribClass_UNDEFINED; break;
+  case SYMS_DwAttribKind_GNU_CALL_SITE_VALUE:            result = SYMS_DwAttribClass_EXPRLOC;   break;
+  case SYMS_DwAttribKind_GNU_CALL_SITE_DATA_VALUE:       result = SYMS_DwAttribClass_EXPRLOC;   break;
+  case SYMS_DwAttribKind_GNU_CALL_SITE_TARGET:           result = SYMS_DwAttribClass_EXPRLOC;   break;
+  case SYMS_DwAttribKind_GNU_CALL_SITE_TARGET_CLOBBERED: result = SYMS_DwAttribClass_EXPRLOC;   break;
+  case SYMS_DwAttribKind_GNU_TAIL_CALL:                  result = SYMS_DwAttribClass_FLAG;      break;
+  case SYMS_DwAttribKind_GNU_ALL_TAIL_CALL_SITES:        result = SYMS_DwAttribClass_FLAG;      break;
+  case SYMS_DwAttribKind_GNU_ALL_CALL_SITES:             result = SYMS_DwAttribClass_FLAG;      break;
+  case SYMS_DwAttribKind_GNU_ALL_SOURCE_CALL_SITES:      result = SYMS_DwAttribClass_FLAG;      break;
+  case SYMS_DwAttribKind_GNU_MACROS:                     result = SYMS_DwAttribClass_FLAG;      break;
+  case SYMS_DwAttribKind_GNU_DELETED:                    result = SYMS_DwAttribClass_UNDEFINED; break;
+  case SYMS_DwAttribKind_GNU_DWO_NAME:                   result = SYMS_DwAttribClass_STRING;    break;
+  case SYMS_DwAttribKind_GNU_DWO_ID:                     result = SYMS_DwAttribClass_CONST;     break;
+  case SYMS_DwAttribKind_GNU_RANGES_BASE:                result = SYMS_DwAttribClass_UNDEFINED; break;
+  case SYMS_DwAttribKind_GNU_ADDR_BASE:                  result = SYMS_DwAttribClass_ADDRPTR;   break;
+  case SYMS_DwAttribKind_GNU_PUBNAMES:                   result = SYMS_DwAttribClass_FLAG;      break;
+  case SYMS_DwAttribKind_GNU_PUBTYPES:                   result = SYMS_DwAttribClass_UNDEFINED; break;
+  case SYMS_DwAttribKind_GNU_DISCRIMINATOR:              result = SYMS_DwAttribClass_CONST;     break;
+  case SYMS_DwAttribKind_GNU_LOCVIEWS:                   result = SYMS_DwAttribClass_UNDEFINED; break;
+  case SYMS_DwAttribKind_GNU_ENTRY_VIEW:                 result = SYMS_DwAttribClass_UNDEFINED; break;
+  case SYMS_DwAttribKind_VMS_RTNBEG_PD_ADDRESS:          result = SYMS_DwAttribClass_UNDEFINED; break;
+  case SYMS_DwAttribKind_USE_GNAT_DESCRIPTIVE_TYPE:      result = SYMS_DwAttribClass_UNDEFINED; break;
+  case SYMS_DwAttribKind_GNAT_DESCRIPTIVE_TYPE:          result = SYMS_DwAttribClass_UNDEFINED; break;
+  case SYMS_DwAttribKind_GNU_NUMERATOR:                  result = SYMS_DwAttribClass_UNDEFINED; break;
+  case SYMS_DwAttribKind_GNU_DENOMINATOR:                result = SYMS_DwAttribClass_UNDEFINED; break;
+  case SYMS_DwAttribKind_GNU_BIAS:                       result = SYMS_DwAttribClass_UNDEFINED; break;
+  case SYMS_DwAttribKind_UPC_THREADS_SCALED:             result = SYMS_DwAttribClass_UNDEFINED; break;
+  case SYMS_DwAttribKind_PGI_LBASE:                      result = SYMS_DwAttribClass_UNDEFINED; break;
+  case SYMS_DwAttribKind_PGI_SOFFSET:                    result = SYMS_DwAttribClass_UNDEFINED; break;
+  case SYMS_DwAttribKind_PGI_LSTRIDE:                    result = SYMS_DwAttribClass_UNDEFINED; break;
+  case SYMS_DwAttribKind_LLVM_INCLUDE_PATH:              result = SYMS_DwAttribClass_STRING;    break;
+  case SYMS_DwAttribKind_LLVM_CONFIG_MACROS:             result = SYMS_DwAttribClass_STRING;    break;
+  case SYMS_DwAttribKind_LLVM_SYSROOT:                   result = SYMS_DwAttribClass_STRING;    break;
+  case SYMS_DwAttribKind_LLVM_API_NOTES:                 result = SYMS_DwAttribClass_STRING;    break;
+  case SYMS_DwAttribKind_LLVM_TAG_OFFSET:                result = SYMS_DwAttribClass_UNDEFINED; break;
+  case SYMS_DwAttribKind_APPLE_OPTIMIZED:                result = SYMS_DwAttribClass_FLAG;      break;
+  case SYMS_DwAttribKind_APPLE_FLAGS:                    result = SYMS_DwAttribClass_FLAG;      break;
+  case SYMS_DwAttribKind_APPLE_ISA:                      result = SYMS_DwAttribClass_FLAG;      break;
+  case SYMS_DwAttribKind_APPLE_BLOCK:                    result = SYMS_DwAttribClass_UNDEFINED; break;
+  case SYMS_DwAttribKind_APPLE_MAJOR_RUNTIME_VERS:       result = SYMS_DwAttribClass_UNDEFINED; break;
+  case SYMS_DwAttribKind_APPLE_RUNTIME_CLASS:            result = SYMS_DwAttribClass_UNDEFINED; break;
+  case SYMS_DwAttribKind_APPLE_OMIT_FRAME_PTR:           result = SYMS_DwAttribClass_FLAG;      break;
+  case SYMS_DwAttribKind_APPLE_PROPERTY_NAME:            result = SYMS_DwAttribClass_UNDEFINED; break;
+  case SYMS_DwAttribKind_APPLE_PROPERTY_GETTER:          result = SYMS_DwAttribClass_UNDEFINED; break;
+  case SYMS_DwAttribKind_APPLE_PROPERTY_SETTER:          result = SYMS_DwAttribClass_UNDEFINED; break;
+  case SYMS_DwAttribKind_APPLE_PROPERTY_ATTRIBUTE:       result = SYMS_DwAttribClass_UNDEFINED; break;
+  case SYMS_DwAttribKind_APPLE_OBJC_COMPLETE_TYPE:       result = SYMS_DwAttribClass_UNDEFINED; break;
+  case SYMS_DwAttribKind_APPLE_PROPERTY:                 result = SYMS_DwAttribClass_UNDEFINED; break;
+  case SYMS_DwAttribKind_APPLE_OBJ_DIRECT:               result = SYMS_DwAttribClass_UNDEFINED; break;
+  case SYMS_DwAttribKind_APPLE_SDK:                      result = SYMS_DwAttribClass_STRING;    break;
+  }
+  return result;
+}
+
+SYMS_API SYMS_DwAttribClass
+syms_dw_attrib_class_from_attrib_kind_v2(SYMS_DwAttribKind attrib)
+{
+  SYMS_DwAttribClass result = 0;
+  switch (attrib) {
+  case SYMS_DwAttribKind_SIBLING:               result = SYMS_DwAttribClass_REFERENCE; break;
+  case SYMS_DwAttribKind_LOCATION:              result = SYMS_DwAttribClass_BLOCK|SYMS_DwAttribClass_CONST; break;
+  case SYMS_DwAttribKind_NAME:                  result = SYMS_DwAttribClass_STRING; break;
+  case SYMS_DwAttribKind_ORDERING:              result = SYMS_DwAttribClass_CONST; break;
+  case SYMS_DwAttribKind_BYTE_SIZE:             result = SYMS_DwAttribClass_CONST; break;
+  case SYMS_DwAttribKind_BIT_OFFSET:            result = SYMS_DwAttribClass_CONST; break;
+  case SYMS_DwAttribKind_BIT_SIZE:              result = SYMS_DwAttribClass_CONST; break;
+  case SYMS_DwAttribKind_STMT_LIST:             result = SYMS_DwAttribClass_CONST; break;
+  case SYMS_DwAttribKind_LOW_PC:                result = SYMS_DwAttribClass_ADDRESS; break;
+  case SYMS_DwAttribKind_HIGH_PC:               result = SYMS_DwAttribClass_ADDRESS; break;
+  case SYMS_DwAttribKind_LANGUAGE:              result = SYMS_DwAttribClass_CONST; break;
+  case SYMS_DwAttribKind_DISCR:                 result = SYMS_DwAttribClass_REFERENCE; break;
+  case SYMS_DwAttribKind_DISCR_VALUE:           result = SYMS_DwAttribClass_CONST; break;
+  case SYMS_DwAttribKind_VISIBILITY:            result = SYMS_DwAttribClass_CONST; break;
+  case SYMS_DwAttribKind_IMPORT:                result = SYMS_DwAttribClass_REFERENCE; break;
+  case SYMS_DwAttribKind_STRING_LENGTH:         result = SYMS_DwAttribClass_BLOCK|SYMS_DwAttribClass_CONST; break;
+  case SYMS_DwAttribKind_COMMON_REFERENCE:      result = SYMS_DwAttribClass_REFERENCE; break;
+  case SYMS_DwAttribKind_COMP_DIR:              result = SYMS_DwAttribClass_STRING; break;
+  case SYMS_DwAttribKind_CONST_VALUE:           result = SYMS_DwAttribClass_STRING|SYMS_DwAttribClass_CONST|SYMS_DwAttribClass_BLOCK; break;
+  case SYMS_DwAttribKind_CONTAINING_TYPE:       result = SYMS_DwAttribClass_REFERENCE; break;
+  case SYMS_DwAttribKind_DEFAULT_VALUE:         result = SYMS_DwAttribClass_REFERENCE; break;
+  case SYMS_DwAttribKind_INLINE:                result = SYMS_DwAttribClass_CONST; break;
+  case SYMS_DwAttribKind_IS_OPTIONAL:           result = SYMS_DwAttribClass_FLAG; break;
+  case SYMS_DwAttribKind_LOWER_BOUND:           result = SYMS_DwAttribClass_CONST|SYMS_DwAttribClass_REFERENCE; break;
+  case SYMS_DwAttribKind_PRODUCER:              result = SYMS_DwAttribClass_STRING; break;
+  case SYMS_DwAttribKind_PROTOTYPED:            result = SYMS_DwAttribClass_FLAG; break;
+  case SYMS_DwAttribKind_RETURN_ADDR:           result = SYMS_DwAttribClass_BLOCK|SYMS_DwAttribClass_CONST; break;
+  case SYMS_DwAttribKind_START_SCOPE:           result = SYMS_DwAttribClass_CONST; break;
+  case SYMS_DwAttribKind_BIT_STRIDE:            result = SYMS_DwAttribClass_CONST; break;
+  case SYMS_DwAttribKind_UPPER_BOUND:           result = SYMS_DwAttribClass_CONST|SYMS_DwAttribClass_REFERENCE; break;
+  case SYMS_DwAttribKind_ABSTRACT_ORIGIN:       result = SYMS_DwAttribClass_REFERENCE; break;
+  case SYMS_DwAttribKind_ACCESSIBILITY:         result = SYMS_DwAttribClass_CONST; break;
+  case SYMS_DwAttribKind_ADDRESS_CLASS:         result = SYMS_DwAttribClass_CONST; break;
+  case SYMS_DwAttribKind_ARTIFICIAL:            result = SYMS_DwAttribClass_FLAG; break;
+  case SYMS_DwAttribKind_BASE_TYPES:            result = SYMS_DwAttribClass_REFERENCE; break;
+  case SYMS_DwAttribKind_CALLING_CONVENTION:    result = SYMS_DwAttribClass_CONST; break;
+  case SYMS_DwAttribKind_COUNT:                 result = SYMS_DwAttribClass_CONST|SYMS_DwAttribClass_REFERENCE; break;
+  case SYMS_DwAttribKind_DATA_MEMBER_LOCATION:  result = SYMS_DwAttribClass_BLOCK|SYMS_DwAttribClass_REFERENCE; break;
+  case SYMS_DwAttribKind_DECL_COLUMN:           result = SYMS_DwAttribClass_CONST; break;
+  case SYMS_DwAttribKind_DECL_FILE:             result = SYMS_DwAttribClass_CONST; break;
+  case SYMS_DwAttribKind_DECL_LINE:             result = SYMS_DwAttribClass_CONST; break;
+  case SYMS_DwAttribKind_DECLARATION:           result = SYMS_DwAttribClass_FLAG; break;
+  case SYMS_DwAttribKind_DISCR_LIST:            result = SYMS_DwAttribClass_BLOCK; break;
+  case SYMS_DwAttribKind_ENCODING:              result = SYMS_DwAttribClass_CONST; break;
+  case SYMS_DwAttribKind_EXTERNAL:              result = SYMS_DwAttribClass_FLAG; break;
+  case SYMS_DwAttribKind_FRAME_BASE:            result = SYMS_DwAttribClass_BLOCK|SYMS_DwAttribClass_CONST; break;
+  case SYMS_DwAttribKind_FRIEND:                result = SYMS_DwAttribClass_REFERENCE; break;
+  case SYMS_DwAttribKind_IDENTIFIER_CASE:       result = SYMS_DwAttribClass_CONST; break;
+  case SYMS_DwAttribKind_MACRO_INFO:            result = SYMS_DwAttribClass_CONST; break;
+  case SYMS_DwAttribKind_NAMELIST_ITEM:         result = SYMS_DwAttribClass_BLOCK; break;
+  case SYMS_DwAttribKind_PRIORITY:              result = SYMS_DwAttribClass_REFERENCE; break;
+  case SYMS_DwAttribKind_SEGMENT:               result = SYMS_DwAttribClass_BLOCK|SYMS_DwAttribClass_CONST; break;
+  case SYMS_DwAttribKind_SPECIFICATION:         result = SYMS_DwAttribClass_REFERENCE; break;
+  case SYMS_DwAttribKind_STATIC_LINK:           result = SYMS_DwAttribClass_BLOCK|SYMS_DwAttribClass_CONST; break;
+  case SYMS_DwAttribKind_TYPE:                  result = SYMS_DwAttribClass_REFERENCE; break;
+  case SYMS_DwAttribKind_USE_LOCATION:          result = SYMS_DwAttribClass_BLOCK|SYMS_DwAttribClass_CONST; break;
+  case SYMS_DwAttribKind_VARIABLE_PARAMETER:    result = SYMS_DwAttribClass_FLAG; break;
+  case SYMS_DwAttribKind_VIRTUALITY:            result = SYMS_DwAttribClass_CONST; break;
+  case SYMS_DwAttribKind_VTABLE_ELEM_LOCATION:  result = SYMS_DwAttribClass_BLOCK|SYMS_DwAttribClass_REFERENCE; break;
+  default:                                      result = syms_dw_attrib_class_from_attrib_kind_extensions(attrib); break;
+  }
+  return result;
+}
+
+SYMS_API SYMS_DwAttribClass
+syms_dw_attrib_class_from_form_kind_v2(SYMS_DwFormKind form)
+{
+  SYMS_DwAttribClass result = 0;
+  switch (form) {
+  case SYMS_DwFormKind_ADDR:      result = SYMS_DwAttribClass_ADDRESS; break;
+  case SYMS_DwFormKind_BLOCK2:    result = SYMS_DwAttribClass_BLOCK; break;
+  case SYMS_DwFormKind_BLOCK4:    result = SYMS_DwAttribClass_BLOCK; break;
+  case SYMS_DwFormKind_DATA2:     result = SYMS_DwAttribClass_CONST; break;
+  case SYMS_DwFormKind_DATA4:     result = SYMS_DwAttribClass_CONST; break;
+  case SYMS_DwFormKind_DATA8:     result = SYMS_DwAttribClass_CONST; break;
+  case SYMS_DwFormKind_STRING:    result = SYMS_DwAttribClass_STRING; break;
+  case SYMS_DwFormKind_BLOCK:     result = SYMS_DwAttribClass_BLOCK; break;
+  case SYMS_DwFormKind_BLOCK1:    result = SYMS_DwAttribClass_BLOCK; break;
+  case SYMS_DwFormKind_DATA1:     result = SYMS_DwAttribClass_CONST; break;
+  case SYMS_DwFormKind_FLAG:      result = SYMS_DwAttribClass_FLAG; break;
+  case SYMS_DwFormKind_SDATA:     result = SYMS_DwAttribClass_CONST; break;
+  case SYMS_DwFormKind_STRP:      result = SYMS_DwAttribClass_STRING; break;
+  case SYMS_DwFormKind_UDATA:     result = SYMS_DwAttribClass_CONST; break;
+  case SYMS_DwFormKind_REF_ADDR:  result = SYMS_DwAttribClass_REFERENCE; break;
+  case SYMS_DwFormKind_REF1:      result = SYMS_DwAttribClass_REFERENCE; break;
+  case SYMS_DwFormKind_REF2:      result = SYMS_DwAttribClass_REFERENCE; break;
+  case SYMS_DwFormKind_REF4:      result = SYMS_DwAttribClass_REFERENCE; break;
+  case SYMS_DwFormKind_REF8:      result = SYMS_DwAttribClass_REFERENCE; break;
+  case SYMS_DwFormKind_REF_UDATA: result = SYMS_DwAttribClass_REFERENCE; break;
+  case SYMS_DwFormKind_INDIRECT:  result = SYMS_DwAttribClass_UNDEFINED;
+  }
+  return result;
+}
+
 #endif // SYMS_DWARF_C
diff --git a/Engine/Source/Runtime/SymsLib/syms/core/dwarf/syms_dwarf.h b/Engine/Source/Runtime/SymsLib/syms/core/dwarf/syms_dwarf.h
--- a/Engine/Source/Runtime/SymsLib/syms/core/dwarf/syms_dwarf.h
+++ b/Engine/Source/Runtime/SymsLib/syms/core/dwarf/syms_dwarf.h
@@ -52,6 +52,10 @@ SYMS_API SYMS_B32      syms_dw_are_attrib_class_and_form_kind_compatible(SYMS_Dw
 SYMS_API SYMS_TypeKind syms_dw_type_kind_from_tag_encoding_size(SYMS_DwTagKind tag_kind, SYMS_DwAttribTypeEncoding encoding, SYMS_U64 size);
 SYMS_API SYMS_TypeModifiers syms_dw_type_modifiers_from_tag_kind(SYMS_DwTagKind tag_kind);
 
+SYMS_API SYMS_DwAttribClass syms_dw_attrib_class_from_attrib_kind_extensions(SYMS_DwAttribKind attrib);
+SYMS_API SYMS_DwAttribClass syms_dw_attrib_class_from_attrib_kind_v2(SYMS_DwAttribKind attrib);
+SYMS_API SYMS_DwAttribClass syms_dw_attrib_class_from_form_kind_v2(SYMS_DwFormKind form);
+
 SYMS_C_LINKAGE_END
 
 #endif // SYMS_DWARF_H
diff --git a/Engine/Source/Runtime/SymsLib/syms/core/dwarf/syms_dwarf_expr.h b/Engine/Source/Runtime/SymsLib/syms/core/dwarf/syms_dwarf_expr.h
--- a/Engine/Source/Runtime/SymsLib/syms/core/dwarf/syms_dwarf_expr.h
+++ b/Engine/Source/Runtime/SymsLib/syms/core/dwarf/syms_dwarf_expr.h
@@ -289,7 +289,7 @@ SYMS_API SYMS_DwSimpleLoc syms_dw_expr__analyze_fast(void *base, SYMS_U64Range r
 // The analyzer also determines if there is any possibility for non-linear
 // flow. Jumps, branches, and call ops all create non-linear flow. An
 // expression that doesn't have non-linear flow is trivially gauranteed to
-// terminate and therefore a good candidate for convertion to a human readable
+// terminate and therefore a good candidate for conversion to a human readable
 // expression.
 //
 // The call config is optional (may be null). If is provided the analysis
diff --git a/Engine/Source/Runtime/SymsLib/syms/core/dwarf/syms_dwarf_parser.c b/Engine/Source/Runtime/SymsLib/syms/core/dwarf/syms_dwarf_parser.c
--- a/Engine/Source/Runtime/SymsLib/syms/core/dwarf/syms_dwarf_parser.c
+++ b/Engine/Source/Runtime/SymsLib/syms/core/dwarf/syms_dwarf_parser.c
@@ -141,7 +141,7 @@ syms_dw_sid_from_info_offset(SYMS_U64 info_offset)
 }
 
 SYMS_API SYMS_DwAttribClass
-syms_dw_pick_attrib_value_class(SYMS_DwAttribKind attrib, SYMS_DwFormKind form_kind)
+syms_dw_pick_attrib_value_class(SYMS_DwLanguage lang, SYMS_DwVersion ver, SYMS_DwAttribKind attrib, SYMS_DwFormKind form_kind)
 {
   // NOTE(rjf): DWARF's spec specifies two mappings:
   // (SYMS_DwAttribKind) => List(SYMS_DwAttribClass)
@@ -150,8 +150,33 @@ syms_dw_pick_attrib_value_class(SYMS_DwAttribKind attrib, SYMS_DwFormKind form_k
   // This function's purpose is to find the overlapping class between an
   // SYMS_DwAttribKind and SYMS_DwFormKind.
   SYMS_DwAttribClass result = 0;
-  SYMS_DwAttribClass attrib_info = syms_dw_attrib_class_from_attrib_kind(attrib);
-  SYMS_DwAttribClass form_info = syms_dw_attrib_class_from_form_kind(form_kind);
+  SYMS_DwAttribClass attrib_info;
+  SYMS_DwAttribClass form_info;
+
+  switch(ver)
+  {
+    case SYMS_DwVersion_V1: break;
+    case SYMS_DwVersion_V2:
+    {
+      attrib_info = syms_dw_attrib_class_from_attrib_kind_v2(attrib);
+      form_info = syms_dw_attrib_class_from_form_kind_v2(form_kind);
+
+      // rust compiler writes version 5 attributes
+      if(lang == SYMS_DwLanguage_Rust && (attrib_info == 0 || form_info == 0))
+      {
+        attrib_info = syms_dw_attrib_class_from_attrib_kind(attrib);
+        form_info = syms_dw_attrib_class_from_form_kind(form_kind);
+      }
+    } break;
+    case SYMS_DwVersion_V3:
+    case SYMS_DwVersion_V4:
+    case SYMS_DwVersion_V5:
+    {
+      attrib_info = syms_dw_attrib_class_from_attrib_kind(attrib);
+      form_info = syms_dw_attrib_class_from_form_kind(form_kind);
+    } break;
+  }
+
   if(attrib_info != 0 && form_info != 0)
   {
     result = SYMS_DwAttribClass_UNDEFINED;
@@ -165,6 +190,7 @@ syms_dw_pick_attrib_value_class(SYMS_DwAttribKind attrib, SYMS_DwFormKind form_k
       }
     }
   }
+
   return result;
 }
 
@@ -219,7 +245,7 @@ syms_dw_symbol_kind_from_tag_stub(SYMS_String8 data, SYMS_DwDbgAccel *dbg, SYMS_
       SYMS_B32 has_location = syms_false;
       SYMS_B32 has_const_val = syms_false;
       SYMS_DwAttribValue location_value = {SYMS_DwSectionKind_Null};
-      SYMS_DwAttribList attribs = syms_dw_attrib_list_from_stub(scratch.arena, data, dbg, resolve_params.addr_size, stub);
+      SYMS_DwAttribList attribs = syms_dw_attrib_list_from_stub(scratch.arena, data, dbg, resolve_params.language, resolve_params.version, resolve_params.addr_size, stub);
       for(SYMS_DwAttribNode *attrib_n = attribs.first;
           attrib_n != 0;
           attrib_n = attrib_n->next)
@@ -1836,6 +1862,7 @@ syms_dw_attrib_value_resolve_params_from_comp_root(SYMS_DwCompRoot *root)
   SYMS_DwAttribValueResolveParams params;
   syms_memzero_struct(&params);
   params.version                  = root->version;
+  params.language                 = root->language;
   params.addr_size                = root->address_size;
   params.containing_unit_info_off = root->info_off;
   params.debug_addrs_base         = root->addrs_base;
@@ -2066,7 +2093,7 @@ syms_dw_range_list_from_high_low_pc_and_ranges_attrib_value(SYMS_Arena *arena, S
 //~ rjf: Tag Parsing
 
 SYMS_API SYMS_DwAttribListParseResult
-syms_dw_parse_attrib_list_from_info_abbrev_offsets(SYMS_Arena *arena, SYMS_String8 data, SYMS_DwDbgAccel *dbg, SYMS_U64 address_size, SYMS_U64 info_off, SYMS_U64 abbrev_off)
+syms_dw_parse_attrib_list_from_info_abbrev_offsets(SYMS_Arena *arena, SYMS_String8 data, SYMS_DwDbgAccel *dbg, SYMS_DwLanguage lang, SYMS_DwVersion ver, SYMS_U64 address_size, SYMS_U64 info_off, SYMS_U64 abbrev_off)
 {
   //- rjf: set up prereqs
   void *info_base = syms_dw_sec_base_from_dbg(data, dbg, SYMS_DwSectionKind_Info);
@@ -2096,7 +2123,7 @@ syms_dw_parse_attrib_list_from_info_abbrev_offsets(SYMS_Arena *arena, SYMS_Strin
     //- rjf: extract attrib info from abbrev
     SYMS_DwAttribKind attrib_kind = (SYMS_DwAttribKind)abbrev.id;
     SYMS_DwFormKind form_kind = (SYMS_DwFormKind)abbrev.sub_kind;
-    SYMS_DwAttribClass attrib_class = syms_dw_pick_attrib_value_class(attrib_kind, form_kind);
+    SYMS_DwAttribClass attrib_class = syms_dw_pick_attrib_value_class(lang, ver, attrib_kind, form_kind);
     
     //- rjf: parse the form value from the file
     SYMS_DwAttribValue form_value;
@@ -2141,6 +2168,8 @@ syms_dw_tag_from_info_offset(SYMS_Arena *arena,
                              SYMS_String8 data,
                              SYMS_DwDbgAccel *dbg,
                              SYMS_DwAbbrevTable abbrev_table,
+                             SYMS_DwLanguage lang,
+                             SYMS_DwVersion ver,
                              SYMS_U64 address_size,
                              SYMS_U64 info_offset)
 {
@@ -2185,7 +2214,7 @@ syms_dw_tag_from_info_offset(SYMS_Arena *arena,
   SYMS_DwAttribList attribs = {0};
   if(good_tag_abbrev)
   {
-    SYMS_DwAttribListParseResult attribs_parse = syms_dw_parse_attrib_list_from_info_abbrev_offsets(arena, data, dbg, address_size, info_read_off, abbrev_read_off);
+    SYMS_DwAttribListParseResult attribs_parse = syms_dw_parse_attrib_list_from_info_abbrev_offsets(arena, data, dbg, lang, ver, address_size, info_read_off, abbrev_read_off);
     attribs_info_off = info_read_off;
     attribs_abbrev_off = abbrev_read_off;
     info_read_off = attribs_parse.max_info_off;
@@ -2228,6 +2257,12 @@ syms_dw_stub_from_tag(SYMS_String8 data, SYMS_DwDbgAccel *dbg, SYMS_DwAttribValu
       {
         SYMS_DwAttribValue value = syms_dw_attrib_value_from_form_value(data, dbg, resolve_params, attrib->form_kind, attrib->value_class, attrib->form_value);
         stub.ref = value.v[0];
+        stub.flags |= SYMS_DwTagStubFlag_HasSpecification;
+      }break;
+      case SYMS_DwAttribKind_ABSTRACT_ORIGIN:
+      {
+        SYMS_DwAttribValue value = syms_dw_attrib_value_from_form_value(data, dbg, resolve_params, attrib->form_kind, attrib->value_class, attrib->form_value);
+        stub.abstract_origin = value.v[0];
       }break;
       case SYMS_DwAttribKind_OBJECT_POINTER:       stub.flags |= SYMS_DwTagStubFlag_HasObjectPointerArg; break;
       case SYMS_DwAttribKind_LOCATION:             stub.flags |= SYMS_DwTagStubFlag_HasLocation; break;
@@ -2371,6 +2406,12 @@ syms_dw_comp_root_from_range(SYMS_Arena *arena, SYMS_String8 data, SYMS_DwDbgAcc
     switch(version)
     {
       default: break;
+      case SYMS_DwVersion_V2: {
+        abbrev_base = 0;
+        next_off += syms_based_range_read(info_base, range, next_off, 4, &abbrev_base);
+        next_off += syms_based_range_read(info_base, range, next_off, 1, &address_size);
+        got_unit_kind = syms_true;
+      } break;
       case SYMS_DwVersion_V3:
       case SYMS_DwVersion_V4:
       {
@@ -2407,7 +2448,7 @@ syms_dw_comp_root_from_range(SYMS_Arena *arena, SYMS_String8 data, SYMS_DwDbgAcc
   if(got_unit_kind)
   {
     SYMS_U64 comp_root_tag_off = range.min + next_off;
-    comp_unit_tag = syms_dw_tag_from_info_offset(scratch.arena, data, dbg, abbrev_table, address_size, comp_root_tag_off);
+    comp_unit_tag = syms_dw_tag_from_info_offset(scratch.arena, data, dbg, abbrev_table, SYMS_DwLanguage_NULL, (SYMS_DwVersion)version, address_size, comp_root_tag_off);
     got_comp_unit_tag = syms_true;
   }
   
@@ -2546,7 +2587,7 @@ syms_dw_comp_root_from_range(SYMS_Arena *arena, SYMS_String8 data, SYMS_DwDbgAcc
   {
     unit.dwo_id = gnu_dwo_id;
   }
-  unit.language      = language;
+  unit.language      = (SYMS_DwLanguage)language;
   unit.name_case     = name_case;
   unit.use_utf8      = use_utf8;
   unit.line_off      = line_base;
@@ -3054,6 +3095,7 @@ syms_dw_unit_accel_from_comp_root(SYMS_Arena *arena, SYMS_String8 data, SYMS_DwD
   unit->address_size = comp_root->address_size;
   unit->base_addr    = comp_root->base_addr;
   unit->addrs_base   = comp_root->addrs_base;
+  unit->language     = comp_root->language;
   unit->abbrev_table.count = comp_root->abbrev_table.count;
   unit->abbrev_table.entries = (SYMS_DwAbbrevTableEntry *)syms_push_string_copy(arena, syms_str8((SYMS_U8 *)comp_root->abbrev_table.entries,
                                                                                                  sizeof(SYMS_DwAbbrevTableEntry) * unit->abbrev_table.count)).str;
@@ -3091,7 +3133,7 @@ syms_dw_unit_accel_from_comp_root(SYMS_Arena *arena, SYMS_String8 data, SYMS_DwD
       SYMS_ArenaTemp per_tag_scratch = syms_get_scratch(conflicts, SYMS_ARRAY_SIZE(conflicts));
       
       //- rjf: parse tag
-      SYMS_DwTag *tag = syms_dw_tag_from_info_offset(per_tag_scratch.arena, data, dbg, comp_root->abbrev_table, comp_root->address_size, info_off);
+      SYMS_DwTag *tag = syms_dw_tag_from_info_offset(per_tag_scratch.arena, data, dbg, comp_root->abbrev_table, comp_root->language, comp_root->version, comp_root->address_size, info_off);
       SYMS_B32 good_tag = tag->abbrev_id != 0;
       info_off = tag->info_range.max;
       if(good_tag)
@@ -3339,7 +3381,7 @@ syms_dw_unit_accel_from_comp_root(SYMS_Arena *arena, SYMS_String8 data, SYMS_DwD
   }
   
   //- rjf: log statistics about which tag kinds were cached
-#if RJF_DEBUG
+#if RJF_DEBUG && 0
   {
     SYMS_LogOpen(SYMS_LogFeature_DwarfTags, unit->uid, log);
     SYMS_Log("$tag_profile\n");
@@ -3599,7 +3641,7 @@ syms_dw_tag_stub_from_sid(SYMS_String8 data, SYMS_DwDbgAccel *dbg, SYMS_DwUnitAc
 {
   SYMS_DwTagStub stub = {0};
   SYMS_ArenaTemp scratch = syms_get_scratch(0, 0);
-  SYMS_DwTag *tag = syms_dw_tag_from_info_offset(scratch.arena, data, dbg, unit->abbrev_table, unit->address_size, sid);
+  SYMS_DwTag *tag = syms_dw_tag_from_info_offset(scratch.arena, data, dbg, unit->abbrev_table, unit->language, unit->version, unit->address_size, sid);
   stub = syms_dw_stub_from_tag(data, dbg, unit->resolve_params, tag);
   if(stub.ref == 0)
   {
@@ -3650,7 +3692,7 @@ syms_dw_children_from_tag_stub(SYMS_Arena *arena, SYMS_String8 data, SYMS_DwDbgA
     SYMS_U64 depth = 0;
     for(SYMS_U64 info_off = stub.children_info_off;;)
     {
-      SYMS_DwTag *child_tag = syms_dw_tag_from_info_offset(scratch.arena, data, dbg, unit->abbrev_table, unit->address_size, info_off);
+      SYMS_DwTag *child_tag = syms_dw_tag_from_info_offset(scratch.arena, data, dbg, unit->abbrev_table, unit->language, unit->version, unit->address_size, info_off);
       info_off = child_tag->info_range.max;
       
       if(depth == 0 && child_tag->kind != SYMS_DwTagKind_NULL)
@@ -3681,9 +3723,9 @@ syms_dw_children_from_tag_stub(SYMS_Arena *arena, SYMS_String8 data, SYMS_DwDbgA
 }
 
 SYMS_API SYMS_DwAttribList
-syms_dw_attrib_list_from_stub(SYMS_Arena *arena, SYMS_String8 data, SYMS_DwDbgAccel *dbg, SYMS_U64 addr_size, SYMS_DwTagStub *stub)
+syms_dw_attrib_list_from_stub(SYMS_Arena *arena, SYMS_String8 data, SYMS_DwDbgAccel *dbg, SYMS_DwLanguage lang, SYMS_DwVersion ver, SYMS_U64 addr_size, SYMS_DwTagStub *stub)
 {
-  SYMS_DwAttribListParseResult parse = syms_dw_parse_attrib_list_from_info_abbrev_offsets(arena, data, dbg, addr_size, stub->attribs_info_off, stub->attribs_abbrev_off);
+  SYMS_DwAttribListParseResult parse = syms_dw_parse_attrib_list_from_info_abbrev_offsets(arena, data, dbg, lang, ver, addr_size, stub->attribs_info_off, stub->attribs_abbrev_off);
   return parse.attribs;
 }
 
@@ -3758,7 +3800,7 @@ syms_dw_mems_accel_from_sid(SYMS_Arena *arena, SYMS_String8 data, SYMS_DwDbgAcce
       idx += 1;
       
       //- rjf: get child attributes
-      SYMS_DwAttribList attribs = syms_dw_attrib_list_from_stub(scratch.arena, data, dbg, unit->address_size, child_stub);
+      SYMS_DwAttribList attribs = syms_dw_attrib_list_from_stub(scratch.arena, data, dbg, unit->language, unit->version, unit->address_size, child_stub);
       
       //- rjf: get info from attributes
       SYMS_String8 name = {0};
@@ -3928,8 +3970,8 @@ syms_dw_mems_accel_from_sid(SYMS_Arena *arena, SYMS_String8 data, SYMS_DwDbgAcce
           
           //- rjf: check for constructor/destructor via best-guess string match
           {
-            SYMS_String8 child_name = syms_dw_file_string_from_sid(data, dbg, unit, child_stub->sid);
-            SYMS_String8 child_ref_name = syms_dw_file_string_from_sid(data, dbg, unit, child_ref_stub.sid);
+            SYMS_String8 child_name = syms_dw_attrib_string_from_sid__unstable(data, dbg, unit, SYMS_DwAttribKind_NAME, child_stub->sid);
+            SYMS_String8 child_ref_name = syms_dw_attrib_string_from_sid__unstable(data, dbg, unit, SYMS_DwAttribKind_NAME, child_ref_stub.sid);
             SYMS_String8 real_child_name = child_name.size ? child_name : child_ref_name;
             SYMS_B32 child_is_destructor = real_child_name.size > 0 && syms_string_match(real_child_name, syms_str8_lit("~"),
                                                                                          SYMS_StringMatchFlag_RightSideSloppy);
@@ -3940,7 +3982,7 @@ syms_dw_mems_accel_from_sid(SYMS_Arena *arena, SYMS_String8 data, SYMS_DwDbgAcce
             }
             real_child_name = syms_string_trunc_symbol_heuristic(real_child_name);
             
-            SYMS_String8 symbol_name = syms_dw_file_string_from_sid(data, dbg, unit, sid);
+            SYMS_String8 symbol_name = syms_dw_attrib_string_from_sid__unstable(data, dbg, unit, SYMS_DwAttribKind_NAME, sid);
             symbol_name = syms_string_trunc_symbol_heuristic(symbol_name);
             if(syms_string_match(symbol_name, real_child_name, 0))
             {
@@ -4052,12 +4094,12 @@ syms_dw_symbol_from_mem_number(SYMS_String8 data, SYMS_DwDbgAccel *dbg,
   return result;
 }
 
-SYMS_API SYMS_EnumInfoArray
-syms_dw_enum_info_array_from_sid(SYMS_Arena *arena, SYMS_String8 data, SYMS_DwDbgAccel *dbg,
-                                 SYMS_DwUnitAccel *unit, SYMS_SymbolID sid)
+SYMS_API SYMS_EnumMemberArray
+syms_dw_enum_member_array_from_sid(SYMS_Arena *arena, SYMS_String8 data, SYMS_DwDbgAccel *dbg,
+                                   SYMS_DwUnitAccel *unit, SYMS_SymbolID sid)
 {
   SYMS_ArenaTemp scratch = syms_get_scratch(&arena, 1);
-  SYMS_EnumInfoArray array;
+  SYMS_EnumMemberArray array;
   syms_memzero_struct(&array);
   
   SYMS_DwTagStub stub = syms_dw_cached_tag_stub_from_sid__parse_fallback(data, dbg, unit, sid);
@@ -4075,7 +4117,7 @@ syms_dw_enum_info_array_from_sid(SYMS_Arena *arena, SYMS_String8 data, SYMS_DwDb
       }
     }
     
-    SYMS_EnumInfo *enum_info = syms_push_array_zero(arena, SYMS_EnumInfo, count);
+    SYMS_EnumMember *enum_members = syms_push_array_zero(arena, SYMS_EnumMember, count);
     SYMS_U64 idx = 0;
     for(SYMS_DwTagStubNode *child_n = children.first; child_n != 0; child_n = child_n->next)
     {
@@ -4086,7 +4128,7 @@ syms_dw_enum_info_array_from_sid(SYMS_Arena *arena, SYMS_String8 data, SYMS_DwDb
         SYMS_U64 val = 0;
         
         SYMS_ArenaTemp temp = syms_arena_temp_begin(scratch.arena);
-        SYMS_DwAttribList attribs = syms_dw_attrib_list_from_stub(temp.arena, data, dbg, unit->address_size, child_stub);
+        SYMS_DwAttribList attribs = syms_dw_attrib_list_from_stub(temp.arena, data, dbg, unit->language, unit->version, unit->address_size, child_stub);
         SYMS_B32 got_name = syms_false;
         SYMS_B32 got_val = syms_false;
         
@@ -4112,8 +4154,8 @@ syms_dw_enum_info_array_from_sid(SYMS_Arena *arena, SYMS_String8 data, SYMS_DwDb
           }
         }
         
-        enum_info[idx].name = syms_push_string_copy(arena, name);
-        enum_info[idx].val = val;
+        enum_members[idx].name = syms_push_string_copy(arena, name);
+        enum_members[idx].val = val;
         
         syms_arena_temp_end(temp);
         idx += 1;
@@ -4121,7 +4163,7 @@ syms_dw_enum_info_array_from_sid(SYMS_Arena *arena, SYMS_String8 data, SYMS_DwDb
     }
     
     array.count = count;
-    array.enum_info = enum_info;
+    array.enum_members = enum_members;
   }
   
   syms_release_scratch(scratch);
@@ -4142,11 +4184,19 @@ syms_dw_containing_type_from_sid(SYMS_String8 data, SYMS_DwDbgAccel *dbg, SYMS_D
   return result;
 }
 
+SYMS_API SYMS_String8
+syms_dw_linkage_name_from_sid(SYMS_Arena *arena, SYMS_String8 data, SYMS_DwDbgAccel *dbg, SYMS_DwUnitAccel *unit, SYMS_SymbolID sid)
+{
+  SYMS_String8 name = syms_dw_attrib_string_from_sid__unstable_chain(data, dbg, unit, SYMS_DwAttribKind_LINKAGE_NAME, sid);
+  SYMS_String8 result = syms_push_string_copy(arena, name);
+  return result;
+}
+
 ////////////////////////////////
 //~ rjf: Full Symbol Info Parsing
 
 SYMS_API SYMS_String8
-syms_dw_file_string_from_sid(SYMS_String8 data, SYMS_DwDbgAccel *dbg, SYMS_DwUnitAccel *unit, SYMS_SymbolID sid)
+syms_dw_attrib_string_from_sid__unstable(SYMS_String8 data, SYMS_DwDbgAccel *dbg, SYMS_DwUnitAccel *unit, SYMS_DwAttribKind kind, SYMS_SymbolID sid)
 {
   SYMS_String8 name = syms_str8_lit("");
   if(sid == SYMS_DWARF_VOID_TYPE_ID)
@@ -4157,11 +4207,11 @@ syms_dw_file_string_from_sid(SYMS_String8 data, SYMS_DwDbgAccel *dbg, SYMS_DwUni
   {
     SYMS_ArenaTemp scratch = syms_get_scratch(0, 0);
     SYMS_DwTagStub stub = syms_dw_cached_tag_stub_from_sid__parse_fallback(data, dbg, unit, sid);
-    SYMS_DwAttribList attribs = syms_dw_attrib_list_from_stub(scratch.arena, data, dbg, unit->address_size, &stub);
+    SYMS_DwAttribList attribs = syms_dw_attrib_list_from_stub(scratch.arena, data, dbg, unit->language, unit->version, unit->address_size, &stub);
     for(SYMS_DwAttribNode *n = attribs.first; n != 0; n = n->next)
     {
       SYMS_DwAttrib *attrib = &n->attrib;
-      if(attrib->attrib_kind == SYMS_DwAttribKind_NAME)
+      if(attrib->attrib_kind == kind)
       {
         SYMS_DwAttribValue value = syms_dw_attrib_value_from_form_value(data, dbg, unit->resolve_params,
                                                                         attrib->form_kind,
@@ -4176,6 +4226,35 @@ syms_dw_file_string_from_sid(SYMS_String8 data, SYMS_DwDbgAccel *dbg, SYMS_DwUni
   return name;
 }
 
+SYMS_API SYMS_String8
+syms_dw_attrib_string_from_sid__unstable_chain(SYMS_String8 data, SYMS_DwDbgAccel *dbg,
+                                               SYMS_DwUnitAccel *unit,
+                                               SYMS_DwAttribKind kind, SYMS_SymbolID sid)
+{
+  SYMS_String8 result = {0};
+  for(SYMS_SymbolID id = sid; id != 0 && result.size == 0;)
+  {
+    result = syms_dw_attrib_string_from_sid__unstable(data, dbg, unit, kind, id);
+    if(result.size == 0)
+    {
+      SYMS_DwTagStub stub = syms_dw_cached_tag_stub_from_sid__parse_fallback(data, dbg, unit, id);
+      if(stub.ref != 0 && stub.flags & SYMS_DwTagStubFlag_HasSpecification)
+      {
+        id = stub.ref;
+      }
+      else if(stub.abstract_origin != 0)
+      {
+        id = stub.abstract_origin;
+      }
+      else
+      {
+        break;
+      }
+    }
+  }
+  return result;
+}
+
 SYMS_API SYMS_SymbolKind
 syms_dw_symbol_kind_from_sid(SYMS_String8 data, SYMS_DwDbgAccel *dbg, SYMS_DwUnitAccel *unit, SYMS_SymbolID sid)
 {
@@ -4188,12 +4267,7 @@ syms_dw_symbol_kind_from_sid(SYMS_String8 data, SYMS_DwDbgAccel *dbg, SYMS_DwUni
 SYMS_API SYMS_String8
 syms_dw_symbol_name_from_sid(SYMS_Arena *arena, SYMS_String8 data, SYMS_DwDbgAccel *dbg, SYMS_DwUnitAccel *unit, SYMS_SymbolID sid)
 {
-  SYMS_String8 name = syms_dw_file_string_from_sid(data, dbg, unit, sid);
-  if(name.size == 0)
-  {
-    SYMS_DwTagStub stub = syms_dw_cached_tag_stub_from_sid__parse_fallback(data, dbg, unit, sid);
-    name = syms_dw_file_string_from_sid(data, dbg, unit, stub.ref);
-  }
+  SYMS_String8 name = syms_dw_attrib_string_from_sid__unstable_chain(data, dbg, unit, SYMS_DwAttribKind_NAME, sid);
   SYMS_String8 result = syms_push_string_copy(arena, name);
   return result;
 }
@@ -4210,7 +4284,7 @@ syms_dw_type_info_from_sid(SYMS_String8 data, SYMS_DwDbgAccel *dbg, SYMS_DwUnitA
     SYMS_SymbolID void_id = syms_dw_sid_from_info_offset(SYMS_DWARF_VOID_TYPE_ID);
     
     //- rjf: get attributes
-    SYMS_DwAttribList attribs = syms_dw_attrib_list_from_stub(scratch.arena, data, dbg, unit->address_size, &stub);
+    SYMS_DwAttribList attribs = syms_dw_attrib_list_from_stub(scratch.arena, data, dbg, unit->language, unit->version, unit->address_size, &stub);
     
     //- rjf: define info from attributes
     SYMS_SymbolID                          type = 0;
@@ -4367,7 +4441,7 @@ syms_dw_type_info_from_sid(SYMS_String8 data, SYMS_DwDbgAccel *dbg, SYMS_DwUnitA
         SYMS_DwTagStub *child = &child_n->stub;
         if(child->kind == SYMS_DwTagKind_SUBRANGE_TYPE)
         {
-          SYMS_DwAttribList child_attribs = syms_dw_attrib_list_from_stub(scratch.arena, data, dbg, unit->address_size, child);
+          SYMS_DwAttribList child_attribs = syms_dw_attrib_list_from_stub(scratch.arena, data, dbg, unit->language, unit->version, unit->address_size, child);
           
           for(SYMS_DwAttribNode *child_attrib_n = child_attribs.first;
               child_attrib_n != 0;
@@ -4406,7 +4480,7 @@ syms_dw_const_info_from_sid(SYMS_String8 data, SYMS_DwDbgAccel *dbg, SYMS_DwUnit
   if(stub.sid != 0)
   {
     //- rjf: get attributes
-    SYMS_DwAttribList attribs = syms_dw_attrib_list_from_stub(scratch.arena, data, dbg, unit->address_size, &stub);
+    SYMS_DwAttribList attribs = syms_dw_attrib_list_from_stub(scratch.arena, data, dbg, unit->language, unit->version, unit->address_size, &stub);
     
     //- rjf: define info from attributes
     SYMS_U64 const_value = 0;
@@ -4455,7 +4529,7 @@ syms_dw_type_from_var_sid(SYMS_String8 data, SYMS_DwDbgAccel *dbg, SYMS_DwUnitAc
   SYMS_DwTagStub stub = syms_dw_cached_tag_stub_from_sid__parse_fallback(data, dbg, unit, id);
   if(stub.sid != 0)
   {
-    SYMS_DwAttribList attribs = syms_dw_attrib_list_from_stub(scratch.arena, data, dbg, unit->address_size, &stub);
+    SYMS_DwAttribList attribs = syms_dw_attrib_list_from_stub(scratch.arena, data, dbg, unit->language, unit->version, unit->address_size, &stub);
     result.uid = unit->uid;
     result.sid = syms_dw_sid_from_info_offset(SYMS_DWARF_VOID_TYPE_ID);
     for(SYMS_DwAttribNode *attrib_n = attribs.first;
@@ -4485,7 +4559,7 @@ syms_dw_voff_from_var_sid(SYMS_String8 data, SYMS_DwDbgAccel *dbg, SYMS_DwUnitAc
   SYMS_DwTagStub stub = syms_dw_cached_tag_stub_from_sid__parse_fallback(data, dbg, unit, sid);
   if(stub.sid != 0)
   {
-    SYMS_DwAttribList attribs = syms_dw_attrib_list_from_stub(scratch.arena, data, dbg, unit->address_size, &stub);
+    SYMS_DwAttribList attribs = syms_dw_attrib_list_from_stub(scratch.arena, data, dbg, unit->language, unit->version, unit->address_size, &stub);
     
     SYMS_DwAttrib *location_attrib = 0;
     for(SYMS_DwAttribNode *attrib_n = attribs.first;
@@ -4577,7 +4651,7 @@ syms_dw_location_from_sid(SYMS_Arena *arena, SYMS_String8 data, SYMS_DwDbgAccel
   SYMS_DwTagStub stub = syms_dw_cached_tag_stub_from_sid__parse_fallback(data, dbg, unit, sid);
   if(stub.sid != 0)
   {
-    SYMS_DwAttribList attribs = syms_dw_attrib_list_from_stub(scratch.arena, data, dbg, unit->address_size, &stub);
+    SYMS_DwAttribList attribs = syms_dw_attrib_list_from_stub(scratch.arena, data, dbg, unit->language, unit->version, unit->address_size, &stub);
     for(SYMS_DwAttribNode *n = attribs.first; n != 0; n = n->next)
     {
       SYMS_DwAttrib *attrib = &n->attrib;
@@ -4615,7 +4689,7 @@ syms_dw_location_ranges_from_sid(SYMS_Arena *arena, SYMS_String8 data, SYMS_DwDb
   SYMS_DwTagStub stub = syms_dw_cached_tag_stub_from_sid__parse_fallback(data, dbg, unit, sid);
   if(stub.sid != 0)
   {
-    SYMS_DwAttribList attribs = syms_dw_attrib_list_from_stub(scratch.arena, data, dbg, unit->address_size, &stub);
+    SYMS_DwAttribList attribs = syms_dw_attrib_list_from_stub(scratch.arena, data, dbg, unit->language, unit->version, unit->address_size, &stub);
     for(SYMS_DwAttribNode *n = attribs.first; n != 0; n = n->next)
     {
       SYMS_DwAttrib *attrib = &n->attrib;
@@ -4782,7 +4856,7 @@ syms_dw_scope_vranges_from_sid(SYMS_Arena *arena, SYMS_String8 data, SYMS_DwDbgA
   //- rjf: get range list
   if(stub.sid != 0)
   {
-    SYMS_DwAttribList attribs = syms_dw_attrib_list_from_stub(scratch.arena, data, dbg, unit->address_size, &stub);
+    SYMS_DwAttribList attribs = syms_dw_attrib_list_from_stub(scratch.arena, data, dbg, unit->language, unit->version, unit->address_size, &stub);
     
     SYMS_U64 low_pc = 0;
     SYMS_U64 high_pc = 0;
@@ -4848,7 +4922,7 @@ syms_dw_sig_info_from_sid(SYMS_Arena *arena, SYMS_String8 data, SYMS_DwDbgAccel
   SYMS_DwTagStub stub = syms_dw_cached_tag_stub_from_sid__parse_fallback(data, dbg, unit, id);
   if(stub.sid != 0)
   {
-    SYMS_DwAttribList attribs = syms_dw_attrib_list_from_stub(scratch.arena, data, dbg, unit->address_size, &stub);
+    SYMS_DwAttribList attribs = syms_dw_attrib_list_from_stub(scratch.arena, data, dbg, unit->language, unit->version, unit->address_size, &stub);
     
     //- rjf: grab return type
     for(SYMS_DwAttribNode *attrib_n = attribs.first;
@@ -4878,7 +4952,7 @@ syms_dw_sig_info_from_sid(SYMS_Arena *arena, SYMS_String8 data, SYMS_DwDbgAccel
       SYMS_DwTagStub *child = &child_n->stub;
       if(child->kind == SYMS_DwTagKind_FORMAL_PARAMETER)
       {
-        SYMS_DwAttribList child_attribs = syms_dw_attrib_list_from_stub(scratch.arena, data, dbg, unit->address_size, child);
+        SYMS_DwAttribList child_attribs = syms_dw_attrib_list_from_stub(scratch.arena, data, dbg, unit->language, unit->version, unit->address_size, child);
         SYMS_SymbolID param_type = syms_dw_sid_from_info_offset(SYMS_DWARF_VOID_TYPE_ID);
         for(SYMS_DwAttribNode *child_attrib_n = child_attribs.first;
             child_attrib_n != 0;
diff --git a/Engine/Source/Runtime/SymsLib/syms/core/dwarf/syms_dwarf_parser.h b/Engine/Source/Runtime/SymsLib/syms/core/dwarf/syms_dwarf_parser.h
--- a/Engine/Source/Runtime/SymsLib/syms/core/dwarf/syms_dwarf_parser.h
+++ b/Engine/Source/Runtime/SymsLib/syms/core/dwarf/syms_dwarf_parser.h
@@ -178,6 +178,7 @@ typedef struct SYMS_DwAttribValueResolveParams SYMS_DwAttribValueResolveParams;
 struct SYMS_DwAttribValueResolveParams
 {
   SYMS_DwVersion version;
+  SYMS_DwLanguage language;
   SYMS_U64 addr_size;                // NOTE(rjf): size in bytes of containing compilation unit's addresses
   SYMS_U64 containing_unit_info_off; // NOTE(rjf): containing compilation unit's offset into the .debug_info section
   SYMS_U64 debug_addrs_base;         // NOTE(rjf): containing compilation unit's offset into the .debug_addrs section       (DWARF V5 ONLY)
@@ -263,7 +264,7 @@ struct SYMS_DwCompRoot
   SYMS_String8 compile_dir;
   SYMS_String8 external_dwo_name;
   SYMS_U64 dwo_id;
-  SYMS_U64 language;
+  SYMS_DwLanguage language;
   SYMS_U64 name_case;
   SYMS_B32 use_utf8;
   SYMS_U64 line_off;
@@ -319,6 +320,7 @@ enum
   SYMS_DwTagStubFlag_HasObjectPointerArg  = (1<<0),
   SYMS_DwTagStubFlag_HasLocation          = (1<<1),
   SYMS_DwTagStubFlag_HasExternal          = (1<<2),
+  SYMS_DwTagStubFlag_HasSpecification     = (1<<3),
 };
 
 typedef struct SYMS_DwTagStub SYMS_DwTagStub;
@@ -330,6 +332,7 @@ struct SYMS_DwTagStub
   SYMS_U64 children_info_off;
   SYMS_U64 attribs_info_off;
   SYMS_U64 attribs_abbrev_off;
+  
   // NOTE(rjf): SYMS_DwAttribKind_SPECIFICATION is tacked onto definitions that
   // are filling out more info about a "prototype". That attribute is a reference
   // that points back at the declaration tag. The declaration tag has the
@@ -338,6 +341,16 @@ struct SYMS_DwTagStub
   // a reference on both, that point back to each other, so it's always easy to
   // get from decl => spec, or from spec => decl.
   SYMS_SymbolID ref;
+  
+  // NOTE(rjf): SYMS_DwAttribKind_ABSTRACT_ORIGIN is tacked onto some definitions
+  // that are used to specify information more specific to inlining, while wanting
+  // to refer to an "abstract" function DIE, that is not specific to any inline
+  // sites. The DWARF generator will not duplicate information across these, so
+  // we will occasionally need to look at an abstract origin to get abstract
+  // information, like name/linkage-name/etc.
+  SYMS_SymbolID abstract_origin;
+  
+  SYMS_U64 _unused_;
 };
 
 typedef struct SYMS_DwTagStubNode SYMS_DwTagStubNode;
@@ -499,6 +512,7 @@ struct SYMS_DwUnitAccel
   SYMS_U64 address_size;
   SYMS_U64 base_addr;
   SYMS_U64 addrs_base;
+  SYMS_DwLanguage language;
   SYMS_DwAbbrevTable abbrev_table;
   
   //- rjf: tag stub hash table
@@ -562,7 +576,7 @@ SYMS_API SYMS_U64           syms_dw_hash_from_string(SYMS_String8 string);
 SYMS_API SYMS_U64           syms_dw_hash_from_sid(SYMS_SymbolID sid);
 
 SYMS_API SYMS_SymbolID      syms_dw_sid_from_info_offset(SYMS_U64 info_offset);
-SYMS_API SYMS_DwAttribClass syms_dw_pick_attrib_value_class(SYMS_DwAttribKind attrib, SYMS_DwFormKind form_kind);
+SYMS_API SYMS_DwAttribClass syms_dw_pick_attrib_value_class(SYMS_DwLanguage lang, SYMS_DwVersion ver, SYMS_DwAttribKind attrib, SYMS_DwFormKind form_kind);
 SYMS_API SYMS_SymbolKind    syms_dw_symbol_kind_from_tag_stub(SYMS_String8 data, SYMS_DwDbgAccel *dbg, SYMS_DwAttribValueResolveParams resolve_params, SYMS_DwTagStub *stub);
 
 SYMS_API SYMS_SecInfoArray syms_dw_copy_sec_info_array(SYMS_Arena *arena, SYMS_SecInfoArray array);
@@ -657,10 +671,12 @@ SYMS_API SYMS_U64RangeList syms_dw_range_list_from_high_low_pc_and_ranges_attrib
 ////////////////////////////////
 //~ rjf: Tag Parsing
 
-SYMS_API SYMS_DwAttribListParseResult syms_dw_parse_attrib_list_from_info_abbrev_offsets(SYMS_Arena *arena, SYMS_String8 data, SYMS_DwDbgAccel *dbg, SYMS_U64 address_size, SYMS_U64 info_off, SYMS_U64 abbrev_off);
+SYMS_API SYMS_DwAttribListParseResult syms_dw_parse_attrib_list_from_info_abbrev_offsets(SYMS_Arena *arena, SYMS_String8 data, SYMS_DwDbgAccel *dbg, SYMS_DwLanguage lang, SYMS_DwVersion ver, SYMS_U64 address_size, SYMS_U64 info_off, SYMS_U64 abbrev_off);
 SYMS_API SYMS_DwTag *syms_dw_tag_from_info_offset(SYMS_Arena *arena,
                                                   SYMS_String8 data, SYMS_DwDbgAccel *dbg,
                                                   SYMS_DwAbbrevTable abbrev_table,
+                                                  SYMS_DwLanguage lang,
+                                                  SYMS_DwVersion ver,
                                                   SYMS_U64 address_size,
                                                   SYMS_U64 info_offset);
 SYMS_API SYMS_DwTagStub syms_dw_stub_from_tag(SYMS_String8 data, SYMS_DwDbgAccel *dbg, SYMS_DwAttribValueResolveParams resolve_params,
@@ -718,7 +734,7 @@ SYMS_API SYMS_DwTagStubCacheNode *  syms_dw_tag_stub_cache_node_from_sid(SYMS_Dw
 SYMS_API SYMS_DwTagStub             syms_dw_tag_stub_from_sid(SYMS_String8 data, SYMS_DwDbgAccel *dbg, SYMS_DwUnitAccel *unit, SYMS_SymbolID sid);
 SYMS_API SYMS_DwTagStub             syms_dw_cached_tag_stub_from_sid__parse_fallback(SYMS_String8 data, SYMS_DwDbgAccel *dbg, SYMS_DwUnitAccel *unit, SYMS_SymbolID sid);
 SYMS_API SYMS_DwTagStubList         syms_dw_children_from_tag_stub(SYMS_Arena *arena, SYMS_String8 data, SYMS_DwDbgAccel *dbg, SYMS_DwUnitAccel *unit, SYMS_DwTagStub stub);
-SYMS_API SYMS_DwAttribList          syms_dw_attrib_list_from_stub(SYMS_Arena *arena, SYMS_String8 data, SYMS_DwDbgAccel *dbg, SYMS_U64 addr_size, SYMS_DwTagStub *stub);
+SYMS_API SYMS_DwAttribList          syms_dw_attrib_list_from_stub(SYMS_Arena *arena, SYMS_String8 data, SYMS_DwDbgAccel *dbg, SYMS_DwLanguage lang, SYMS_DwVersion ver, SYMS_U64 addr_size, SYMS_DwTagStub *stub);
 SYMS_API SYMS_SymbolIDArray         syms_dw_copy_sid_array_if_needed(SYMS_Arena *arena, SYMS_SymbolIDArray arr);
 SYMS_API SYMS_SymbolIDArray         syms_dw_proc_sid_array_from_unit(SYMS_Arena *arena, SYMS_DwUnitAccel *unit);
 SYMS_API SYMS_SymbolIDArray         syms_dw_var_sid_array_from_unit(SYMS_Arena *arena, SYMS_DwUnitAccel *unit);
@@ -743,15 +759,20 @@ SYMS_API SYMS_SigInfo      syms_dw_sig_info_from_mem_number(SYMS_Arena *arena, S
 SYMS_API SYMS_USID         syms_dw_symbol_from_mem_number(SYMS_String8 data, SYMS_DwDbgAccel *dbg,
                                                           SYMS_DwUnitAccel *unit, SYMS_DwMemsAccel *mems,
                                                           SYMS_U64 n);
-SYMS_API SYMS_EnumInfoArray syms_dw_enum_info_array_from_sid(SYMS_Arena *arena, SYMS_String8 data, SYMS_DwDbgAccel *dbg,
-                                                             SYMS_DwUnitAccel *unit, SYMS_SymbolID sid);
+SYMS_API SYMS_EnumMemberArray syms_dw_enum_member_array_from_sid(SYMS_Arena *arena, SYMS_String8 data, SYMS_DwDbgAccel *dbg,
+                                                                 SYMS_DwUnitAccel *unit, SYMS_SymbolID sid);
 SYMS_API SYMS_USID          syms_dw_containing_type_from_sid(SYMS_String8 data, SYMS_DwDbgAccel *dbg, SYMS_DwUnitAccel *unit, SYMS_SymbolID sid);
+SYMS_API SYMS_String8       syms_dw_linkage_name_from_sid(SYMS_Arena *arena, SYMS_String8 data, SYMS_DwDbgAccel *dbg, SYMS_DwUnitAccel *unit, SYMS_SymbolID sid);
 
 ////////////////////////////////
 //~ rjf: Full Symbol Info Parsing
 
-SYMS_API SYMS_String8      syms_dw_file_string_from_sid(SYMS_String8 data, SYMS_DwDbgAccel *dbg,
-                                                        SYMS_DwUnitAccel *unit, SYMS_SymbolID sid);
+SYMS_API SYMS_String8      syms_dw_attrib_string_from_sid__unstable(SYMS_String8 data, SYMS_DwDbgAccel *dbg,
+                                                                    SYMS_DwUnitAccel *unit,
+                                                                    SYMS_DwAttribKind kind, SYMS_SymbolID sid);
+SYMS_API SYMS_String8      syms_dw_attrib_string_from_sid__unstable_chain(SYMS_String8 data, SYMS_DwDbgAccel *dbg,
+                                                                          SYMS_DwUnitAccel *unit,
+                                                                          SYMS_DwAttribKind kind, SYMS_SymbolID sid);
 SYMS_API SYMS_SymbolKind   syms_dw_symbol_kind_from_sid(SYMS_String8 data, SYMS_DwDbgAccel *dbg,
                                                         SYMS_DwUnitAccel *unit, SYMS_SymbolID sid);
 SYMS_API SYMS_String8      syms_dw_symbol_name_from_sid(SYMS_Arena *arena, SYMS_String8 data, SYMS_DwDbgAccel *dbg,
