Upgrade SymsLib to latest version

diff --git a/Engine/Source/Runtime/SymsLib/syms/core/dwarf/syms_dwarf_transpiler.c b/Engine/Source/Runtime/SymsLib/syms/core/dwarf/syms_dwarf_transpiler.c
--- a/Engine/Source/Runtime/SymsLib/syms/core/dwarf/syms_dwarf_transpiler.c
+++ b/Engine/Source/Runtime/SymsLib/syms/core/dwarf/syms_dwarf_transpiler.c
@@ -69,7 +69,7 @@ syms_dw_expr__ir_encode_u(SYMS_Arena *arena, SYMS_DwEvalIRGraph *graph, SYMS_U64
     op = SYMS_EvalOp_ConstU32;
   }
   
-  SYMS_DwEvalIRGraphNode *result = syms_dw_expr__ir_push_node(arena, graph, op, syms_op_params(u));
+  SYMS_DwEvalIRGraphNode *result = syms_dw_expr__ir_push_node(arena, graph, op, syms_eval_op_params(u));
   return(result);
 }
 
@@ -91,9 +91,9 @@ syms_dw_expr__ir_encode_s(SYMS_Arena *arena, SYMS_DwEvalIRGraph *graph, SYMS_S64
     size = 32;
   }
   
-  SYMS_DwEvalIRGraphNode *result = syms_dw_expr__ir_push_node(arena, graph, op, syms_op_params(s));
+  SYMS_DwEvalIRGraphNode *result = syms_dw_expr__ir_push_node(arena, graph, op, syms_eval_op_params(s));
   if (size < 64){
-    syms_dw_expr__ir_push_node(arena, graph, SYMS_EvalOp_TruncSigned, syms_op_params(size));
+    syms_dw_expr__ir_push_node(arena, graph, SYMS_EvalOp_TruncSigned, syms_eval_op_params(size));
   }
   return(result);
 }
@@ -185,12 +185,13 @@ syms_dw_expr__ir_graph_from_dw_expr(SYMS_Arena *arena, SYMS_DwDbgAccel *dbg,
           SYMS_U64 offset = raw_offset + text_voff;
           
           if (offset < 0xFFFFFFFF){
-            new_node = syms_dw_expr__ir_push_node(arena, &graph, SYMS_EvalOp_ModuleOff, syms_op_params(offset));
+            new_node = syms_dw_expr__ir_push_node(arena, &graph, SYMS_EvalOp_ModuleOff, syms_eval_op_params(offset));
           }
           else{
-            new_node = syms_dw_expr__ir_push_node(arena, &graph, SYMS_EvalOp_ModuleOff, syms_op_params(0xFFFFFFFF));
+            new_node = syms_dw_expr__ir_push_node(arena, &graph, SYMS_EvalOp_ModuleOff,
+                                                  syms_eval_op_params(0xFFFFFFFF));
             syms_dw_expr__ir_encode_u(arena, &graph, offset - 0xFFFFFFFF);
-            syms_dw_expr__ir_push_node(arena, &graph, SYMS_EvalOp_Add, syms_op_params(SYMS_EvalTypeGroup_U));
+            syms_dw_expr__ir_push_node(arena, &graph, SYMS_EvalOp_Add, syms_eval_op_params(SYMS_EvalTypeGroup_U));
           }
         }break;
         
@@ -217,12 +218,13 @@ syms_dw_expr__ir_graph_from_dw_expr(SYMS_Arena *arena, SYMS_DwDbgAccel *dbg,
           step_cursor += syms_based_range_read_sleb128(expr_base, expr_range, step_cursor, (SYMS_S64*)&offset);
           
           if (offset < 256){
-            new_node = syms_dw_expr__ir_push_node(arena, &graph, SYMS_EvalOp_FrameOff, syms_op_params((SYMS_U64)offset));
+            new_node = syms_dw_expr__ir_push_node(arena, &graph, SYMS_EvalOp_FrameOff,
+                                                  syms_eval_op_params((SYMS_U64)offset));
           }
           else{
-            new_node = syms_dw_expr__ir_push_node(arena, &graph, SYMS_EvalOp_FrameOff, syms_op_params(255));
+            new_node = syms_dw_expr__ir_push_node(arena, &graph, SYMS_EvalOp_FrameOff, syms_eval_op_params(255));
             syms_dw_expr__ir_encode_u(arena, &graph, offset - 255);
-            syms_dw_expr__ir_push_node(arena, &graph, SYMS_EvalOp_Add, syms_op_params(SYMS_EvalTypeGroup_U));
+            syms_dw_expr__ir_push_node(arena, &graph, SYMS_EvalOp_Add, syms_eval_op_params(SYMS_EvalTypeGroup_U));
           }
         }break;
         
@@ -247,9 +249,9 @@ syms_dw_expr__ir_graph_from_dw_expr(SYMS_Arena *arena, SYMS_DwDbgAccel *dbg,
           // TODO(allen): change size based on arch
           SYMS_U8 size = 8;
           
-          new_node = syms_dw_expr__ir_push_node(arena, &graph, SYMS_EvalOp_RegRead, syms_op_params_2u8(reg_id, size));
+          new_node = syms_dw_expr__ir_push_node(arena, &graph, SYMS_EvalOp_RegRead, syms_eval_op_params_2u8(reg_id, size));
           syms_dw_expr__ir_encode_u(arena, &graph, offset);
-          syms_dw_expr__ir_push_node(arena, &graph, SYMS_EvalOp_Add, syms_op_params(SYMS_EvalTypeGroup_U));
+          syms_dw_expr__ir_push_node(arena, &graph, SYMS_EvalOp_Add, syms_eval_op_params(SYMS_EvalTypeGroup_U));
         }break;
         
         case SYMS_DwOp_BREGX:
@@ -263,22 +265,22 @@ syms_dw_expr__ir_graph_from_dw_expr(SYMS_Arena *arena, SYMS_DwDbgAccel *dbg,
           // TODO(allen): change size based on arch
           SYMS_U8 size = 8;
           
-          new_node = syms_dw_expr__ir_push_node(arena, &graph, SYMS_EvalOp_RegRead, syms_op_params_2u8(reg_id, size));
+          new_node = syms_dw_expr__ir_push_node(arena, &graph, SYMS_EvalOp_RegRead, syms_eval_op_params_2u8(reg_id, size));
           
           syms_dw_expr__ir_encode_u(arena, &graph, offset);
-          syms_dw_expr__ir_push_node(arena, &graph, SYMS_EvalOp_Add, syms_op_params(SYMS_EvalTypeGroup_U));
+          syms_dw_expr__ir_push_node(arena, &graph, SYMS_EvalOp_Add, syms_eval_op_params(SYMS_EvalTypeGroup_U));
         }break;
         
         //// stack operations ////
         
         case SYMS_DwOp_DUP:
         {
-          new_node = syms_dw_expr__ir_push_node(arena, &graph, SYMS_EvalOp_Pick, syms_op_params(0));
+          new_node = syms_dw_expr__ir_push_node(arena, &graph, SYMS_EvalOp_Pick, syms_eval_op_params(0));
         }break;
         
         case SYMS_DwOp_DROP:
         {
-          new_node = syms_dw_expr__ir_push_node(arena, &graph, SYMS_EvalOp_Pop, syms_op_params(0));
+          new_node = syms_dw_expr__ir_push_node(arena, &graph, SYMS_EvalOp_Pop, syms_eval_op_params(0));
         }break;
         
         case SYMS_DwOp_PICK:
@@ -286,27 +288,27 @@ syms_dw_expr__ir_graph_from_dw_expr(SYMS_Arena *arena, SYMS_DwDbgAccel *dbg,
           SYMS_U64 idx = 0;
           step_cursor += syms_based_range_read(expr_base, expr_range, step_cursor, 1, &idx);
           // TODO: if (idx >= 256) no-convert
-          new_node = syms_dw_expr__ir_push_node(arena, &graph, SYMS_EvalOp_Pick, syms_op_params(idx));
+          new_node = syms_dw_expr__ir_push_node(arena, &graph, SYMS_EvalOp_Pick, syms_eval_op_params(idx));
         }break;
         
         case SYMS_DwOp_OVER:
         {
-          new_node = syms_dw_expr__ir_push_node(arena, &graph, SYMS_EvalOp_Pick, syms_op_params(1));
+          new_node = syms_dw_expr__ir_push_node(arena, &graph, SYMS_EvalOp_Pick, syms_eval_op_params(1));
         }break;
         
         case SYMS_DwOp_SWAP:
         {
-          new_node = syms_dw_expr__ir_push_node(arena, &graph, SYMS_EvalOp_Insert, syms_op_params(1));
+          new_node = syms_dw_expr__ir_push_node(arena, &graph, SYMS_EvalOp_Insert, syms_eval_op_params(1));
         }break;
         
         case SYMS_DwOp_ROT:
         {
-          new_node = syms_dw_expr__ir_push_node(arena, &graph, SYMS_EvalOp_Insert, syms_op_params(2));
+          new_node = syms_dw_expr__ir_push_node(arena, &graph, SYMS_EvalOp_Insert, syms_eval_op_params(2));
         }break;
         
         case SYMS_DwOp_DEREF:
         {
-          new_node = syms_dw_expr__ir_push_node(arena, &graph, SYMS_EvalOp_MemRead, syms_op_params(8));
+          new_node = syms_dw_expr__ir_push_node(arena, &graph, SYMS_EvalOp_MemRead, syms_eval_op_params(8));
         }break;
         
         case SYMS_DwOp_DEREF_SIZE:
@@ -314,7 +316,7 @@ syms_dw_expr__ir_graph_from_dw_expr(SYMS_Arena *arena, SYMS_DwDbgAccel *dbg,
           SYMS_U64 raw_size = 0;
           step_cursor += syms_based_range_read(expr_base, expr_range, step_cursor, 1, &raw_size);
           SYMS_U64 size = SYMS_ClampTop(raw_size, 8);
-          new_node = syms_dw_expr__ir_push_node(arena, &graph, SYMS_EvalOp_MemRead, syms_op_params(size));
+          new_node = syms_dw_expr__ir_push_node(arena, &graph, SYMS_EvalOp_MemRead, syms_eval_op_params(size));
         }break;
         
         case SYMS_DwOp_XDEREF:
@@ -325,7 +327,7 @@ syms_dw_expr__ir_graph_from_dw_expr(SYMS_Arena *arena, SYMS_DwDbgAccel *dbg,
         
         case SYMS_DwOp_PUSH_OBJECT_ADDRESS:
         {
-          new_node = syms_dw_expr__ir_push_node(arena, &graph, SYMS_EvalOp_ObjectOff, syms_op_params(0));
+          new_node = syms_dw_expr__ir_push_node(arena, &graph, SYMS_EvalOp_ObjectOff, syms_eval_op_params(0));
         }break;
         
         case SYMS_DwOp_FORM_TLS_ADDRESS:
@@ -335,7 +337,7 @@ syms_dw_expr__ir_graph_from_dw_expr(SYMS_Arena *arena, SYMS_DwDbgAccel *dbg,
         
         case SYMS_DwOp_CALL_FRAME_CFA:
         {
-          new_node = syms_dw_expr__ir_push_node(arena, &graph, SYMS_EvalOp_CFA, syms_op_params(0));
+          new_node = syms_dw_expr__ir_push_node(arena, &graph, SYMS_EvalOp_CFA, syms_eval_op_params(0));
         }break;
         
         //// arithmetic and logical operations ////
@@ -343,60 +345,60 @@ syms_dw_expr__ir_graph_from_dw_expr(SYMS_Arena *arena, SYMS_DwDbgAccel *dbg,
         case SYMS_DwOp_ABS:
         {
           new_node = syms_dw_expr__ir_push_node(arena, &graph, SYMS_EvalOp_Abs,
-                                                syms_op_params(SYMS_EvalTypeGroup_S));
+                                                syms_eval_op_params(SYMS_EvalTypeGroup_S));
         }break;
         
         case SYMS_DwOp_AND:
         {
           new_node = syms_dw_expr__ir_push_node(arena, &graph, SYMS_EvalOp_BitAnd,
-                                                syms_op_params(SYMS_EvalTypeGroup_U));
+                                                syms_eval_op_params(SYMS_EvalTypeGroup_U));
         }break;
         
         case SYMS_DwOp_DIV:
         {
           new_node = syms_dw_expr__ir_push_node(arena, &graph, SYMS_EvalOp_Div,
-                                                syms_op_params(SYMS_EvalTypeGroup_S));
+                                                syms_eval_op_params(SYMS_EvalTypeGroup_S));
         }break;
         
         case SYMS_DwOp_MINUS:
         {
           new_node = syms_dw_expr__ir_push_node(arena, &graph, SYMS_EvalOp_Sub,
-                                                syms_op_params(SYMS_EvalTypeGroup_U));
+                                                syms_eval_op_params(SYMS_EvalTypeGroup_U));
         }break;
         
         case SYMS_DwOp_MOD:
         {
           new_node = syms_dw_expr__ir_push_node(arena, &graph, SYMS_EvalOp_Mod,
-                                                syms_op_params(SYMS_EvalTypeGroup_U));
+                                                syms_eval_op_params(SYMS_EvalTypeGroup_U));
         }break;
         
         case SYMS_DwOp_MUL:
         {
           new_node = syms_dw_expr__ir_push_node(arena, &graph, SYMS_EvalOp_Mul,
-                                                syms_op_params(SYMS_EvalTypeGroup_U));
+                                                syms_eval_op_params(SYMS_EvalTypeGroup_U));
         }break;
         
         case SYMS_DwOp_NEG:
         {
           new_node = syms_dw_expr__ir_push_node(arena, &graph, SYMS_EvalOp_Neg,
-                                                syms_op_params(SYMS_EvalTypeGroup_S));
+                                                syms_eval_op_params(SYMS_EvalTypeGroup_S));
         }break;
         
         case SYMS_DwOp_NOT:
         {
           new_node = syms_dw_expr__ir_push_node(arena, &graph, SYMS_EvalOp_BitNot,
-                                                syms_op_params(SYMS_EvalTypeGroup_U));
+                                                syms_eval_op_params(SYMS_EvalTypeGroup_U));
         }break;
         
         case SYMS_DwOp_OR:
         {
           new_node = syms_dw_expr__ir_push_node(arena, &graph, SYMS_EvalOp_BitOr,
-                                                syms_op_params(SYMS_EvalTypeGroup_U));
+                                                syms_eval_op_params(SYMS_EvalTypeGroup_U));
         }break;
         
         case SYMS_DwOp_PLUS:
         {
-          syms_dw_expr__ir_push_node(arena, &graph, SYMS_EvalOp_Add, syms_op_params(SYMS_EvalTypeGroup_U));
+          syms_dw_expr__ir_push_node(arena, &graph, SYMS_EvalOp_Add, syms_eval_op_params(SYMS_EvalTypeGroup_U));
         }break;
         
         case SYMS_DwOp_PLUS_UCONST:
@@ -404,31 +406,31 @@ syms_dw_expr__ir_graph_from_dw_expr(SYMS_Arena *arena, SYMS_DwDbgAccel *dbg,
           SYMS_U64 y = 0;
           step_cursor += syms_based_range_read_uleb128(expr_base, expr_range, step_cursor, &y);
           new_node = syms_dw_expr__ir_encode_u(arena, &graph, y);
-          syms_dw_expr__ir_push_node(arena, &graph, SYMS_EvalOp_Add, syms_op_params(SYMS_EvalTypeGroup_U));
+          syms_dw_expr__ir_push_node(arena, &graph, SYMS_EvalOp_Add, syms_eval_op_params(SYMS_EvalTypeGroup_U));
         }break;
         
         case SYMS_DwOp_SHL:
         {
           new_node = syms_dw_expr__ir_push_node(arena, &graph, SYMS_EvalOp_LShift,
-                                                syms_op_params(SYMS_EvalTypeGroup_U));
+                                                syms_eval_op_params(SYMS_EvalTypeGroup_U));
         }break;
         
         case SYMS_DwOp_SHR:
         {
           new_node = syms_dw_expr__ir_push_node(arena, &graph, SYMS_EvalOp_RShift,
-                                                syms_op_params(SYMS_EvalTypeGroup_U));
+                                                syms_eval_op_params(SYMS_EvalTypeGroup_U));
         }break;
         
         case SYMS_DwOp_SHRA:
         {
           new_node = syms_dw_expr__ir_push_node(arena, &graph, SYMS_EvalOp_RShift,
-                                                syms_op_params(SYMS_EvalTypeGroup_S));
+                                                syms_eval_op_params(SYMS_EvalTypeGroup_S));
         }break;
         
         case SYMS_DwOp_XOR:
         {
           new_node = syms_dw_expr__ir_push_node(arena, &graph, SYMS_EvalOp_BitXor,
-                                                syms_op_params(SYMS_EvalTypeGroup_U));
+                                                syms_eval_op_params(SYMS_EvalTypeGroup_U));
         }break;
         
         //// control flow operations ////
@@ -436,44 +438,44 @@ syms_dw_expr__ir_graph_from_dw_expr(SYMS_Arena *arena, SYMS_DwDbgAccel *dbg,
         case SYMS_DwOp_LE:
         {
           new_node = syms_dw_expr__ir_push_node(arena, &graph, SYMS_EvalOp_LsEq,
-                                                syms_op_params(SYMS_EvalTypeGroup_S));
+                                                syms_eval_op_params(SYMS_EvalTypeGroup_S));
         }break;
         
         case SYMS_DwOp_GE:
         {
           new_node = syms_dw_expr__ir_push_node(arena, &graph, SYMS_EvalOp_GrEq,
-                                                syms_op_params(SYMS_EvalTypeGroup_S));
+                                                syms_eval_op_params(SYMS_EvalTypeGroup_S));
         }break;
         
         case SYMS_DwOp_EQ:
         {
           new_node = syms_dw_expr__ir_push_node(arena, &graph, SYMS_EvalOp_EqEq,
-                                                syms_op_params(SYMS_EvalTypeGroup_S));
+                                                syms_eval_op_params(SYMS_EvalTypeGroup_S));
         }break;
         
         case SYMS_DwOp_LT:
         {
           new_node = syms_dw_expr__ir_push_node(arena, &graph, SYMS_EvalOp_Less,
-                                                syms_op_params(SYMS_EvalTypeGroup_S));
+                                                syms_eval_op_params(SYMS_EvalTypeGroup_S));
         }break;
         
         case SYMS_DwOp_GT:
         {
           new_node = syms_dw_expr__ir_push_node(arena, &graph, SYMS_EvalOp_Grtr,
-                                                syms_op_params(SYMS_EvalTypeGroup_S));
+                                                syms_eval_op_params(SYMS_EvalTypeGroup_S));
         }break;
         
         case SYMS_DwOp_NE:
         {
           new_node = syms_dw_expr__ir_push_node(arena, &graph, SYMS_EvalOp_NtEq,
-                                                syms_op_params(SYMS_EvalTypeGroup_S));
+                                                syms_eval_op_params(SYMS_EvalTypeGroup_S));
         }break;
         
         case SYMS_DwOp_SKIP:
         {
           SYMS_S16 d = 0;
           step_cursor += syms_based_range_read(expr_base, expr_range, step_cursor, 2, &d);
-          new_node = syms_dw_expr__ir_push_node(arena, &graph, SYMS_EvalOp_Skip, syms_op_params(0));
+          new_node = syms_dw_expr__ir_push_node(arena, &graph, SYMS_EvalOp_Skip, syms_eval_op_params(0));
           
           // save a jump patch
           SYMS_U64 dst_off = step_cursor + d;
@@ -488,7 +490,7 @@ syms_dw_expr__ir_graph_from_dw_expr(SYMS_Arena *arena, SYMS_DwDbgAccel *dbg,
         {
           SYMS_S16 d = 0;
           step_cursor += syms_based_range_read(expr_base, expr_range, step_cursor, 2, &d);
-          new_node = syms_dw_expr__ir_push_node(arena, &graph, SYMS_EvalOp_Cond, syms_op_params(0));
+          new_node = syms_dw_expr__ir_push_node(arena, &graph, SYMS_EvalOp_Cond, syms_eval_op_params(0));
           
           // save a jump patch
           SYMS_U64 dst_off = step_cursor + d;
@@ -511,7 +513,7 @@ syms_dw_expr__ir_graph_from_dw_expr(SYMS_Arena *arena, SYMS_DwDbgAccel *dbg,
         case SYMS_DwOp_NOP:
         {
           new_node = syms_dw_expr__ir_push_node(arena, &graph, (SYMS_EvalOp)SYMS_EvalIRExtKind_Noop,
-                                                syms_op_params(0));
+                                                syms_eval_op_params(0));
         }break;
         
         //// register location descriptions ////
diff --git a/Engine/Source/Runtime/SymsLib/syms/core/elf/syms_elf_parser.c b/Engine/Source/Runtime/SymsLib/syms/core/elf/syms_elf_parser.c
--- a/Engine/Source/Runtime/SymsLib/syms/core/elf/syms_elf_parser.c
+++ b/Engine/Source/Runtime/SymsLib/syms/core/elf/syms_elf_parser.c
@@ -9,8 +9,6 @@
 SYMS_API SYMS_ElfImgHeader
 syms_elf_img_header_from_file(SYMS_String8 file)
 {
-  SYMS_ElfImgHeader img;
-  syms_memzero_struct(&img);
   void *file_base = file.str;
   SYMS_U64Range file_range = syms_make_u64_range(0, file.size);
   
@@ -30,36 +28,36 @@ syms_elf_img_header_from_file(SYMS_String8 file)
   }
   
   //- rjf: parse ELF header
-  SYMS_ElfEhdr64 elf_header;
+  SYMS_ElfEhdr64 elf_header = {0};
   SYMS_B32 is_32bit = syms_false;
   SYMS_B32 good_elf_header = syms_false;
   if(sig_is_elf)
   {
-    syms_memzero_struct(&elf_header);
     SYMS_U64 bytes_read = 0;
     switch(sig[SYMS_ElfIdentifier_CLASS])
     {
       //- rjf: parse 32-bit header
       case SYMS_ElfClass_32:
       {
-        SYMS_ElfEhdr32 elf_header_32;
-        syms_memzero_struct(&elf_header_32);
+        SYMS_ElfEhdr32 elf_header_32 = {0};
         bytes_read = syms_based_range_read_struct(file_base, file_range, elf_header_off, &elf_header_32);
         elf_header = syms_elf_ehdr64_from_ehdr32(elf_header_32);
-        good_elf_header = bytes_read == sizeof(elf_header_32);
+        good_elf_header = (bytes_read == sizeof(elf_header_32));
         is_32bit = syms_true;
       }break;
       //- rjf: parse 64-bit header
       case SYMS_ElfClass_64:
       {
-        syms_memzero_struct(&elf_header);
         bytes_read = syms_based_range_read_struct(file_base, file_range, elf_header_off, &elf_header);
-        good_elf_header = bytes_read == sizeof(elf_header);
+        good_elf_header = (bytes_read == sizeof(elf_header));
       }break;
       default:break;
     }
   }
   
+  //- allen: extract entry point
+  SYMS_U64 entry_point = elf_header.e_entry;
+  
   //- rjf: parse section header
   SYMS_U64 sh_name_low_offset  = SYMS_U64_MAX;
   SYMS_U64 sh_name_high_offset = SYMS_U64_MAX;
@@ -99,8 +97,7 @@ syms_elf_img_header_from_file(SYMS_String8 file)
     //- rjf: search for base address, by grabbing the first LOAD phdr
     for(SYMS_U64 i = 0; i < program_header_num; i += 1)
     {
-      SYMS_ElfPhdr64 program_header;
-      syms_memzero_struct(&program_header);
+      SYMS_ElfPhdr64 program_header = {0};
       switch(sig[SYMS_ElfIdentifier_CLASS])
       {
         case SYMS_ElfClass_32:
@@ -138,6 +135,7 @@ syms_elf_img_header_from_file(SYMS_String8 file)
   }
   
   //- rjf: fill img
+  SYMS_ElfImgHeader img = {0};
   if(good_elf_header)
   {
     img.valid = 1;
@@ -147,6 +145,7 @@ syms_elf_img_header_from_file(SYMS_String8 file)
     img.sh_name_low_offset = sh_name_low_offset;
     img.sh_name_high_offset = sh_name_high_offset;
     img.base_address = base_address;
+    img.entry_point = entry_point;
   }
   
   return img;
@@ -308,8 +307,7 @@ syms_elf_sec_name_from_elf_section(SYMS_ElfSection elf_section){
 SYMS_API SYMS_ElfFileAccel *
 syms_elf_file_accel_from_data(SYMS_Arena *arena, SYMS_String8 string)
 {
-  SYMS_ElfFileAccel *file_accel = syms_push_array(arena, SYMS_ElfFileAccel, 1);
-  syms_memzero_struct(file_accel);
+  SYMS_ElfFileAccel *file_accel = syms_push_array_zero(arena, SYMS_ElfFileAccel, 1);
   file_accel->header = syms_elf_img_header_from_file(string);
   if (file_accel->header.valid){
     file_accel->format = SYMS_FileFormat_ELF;
@@ -323,8 +321,7 @@ syms_elf_file_accel_from_data(SYMS_Arena *arena, SYMS_String8 string)
 SYMS_API SYMS_ElfBinAccel *
 syms_elf_bin_accel_from_file(SYMS_Arena *arena, SYMS_String8 data, SYMS_ElfFileAccel *file_accel)
 {
-  SYMS_ElfBinAccel *bin_accel = syms_push_array(arena, SYMS_ElfBinAccel, 1);
-  syms_memzero_struct(bin_accel);
+  SYMS_ElfBinAccel *bin_accel = syms_push_array_zero(arena, SYMS_ElfBinAccel, 1);
   syms_memmove(&bin_accel->header, &file_accel->header, sizeof(file_accel->header));
   bin_accel->format = file_accel->format;
   bin_accel->sections = syms_elf_section_array_from_img_header(arena, data, file_accel->header);
@@ -335,8 +332,7 @@ syms_elf_bin_accel_from_file(SYMS_Arena *arena, SYMS_String8 data, SYMS_ElfFileA
 SYMS_API SYMS_ExtFileList
 syms_elf_ext_file_list_from_bin(SYMS_Arena *arena, SYMS_String8 file, SYMS_ElfBinAccel *bin_accel)
 {
-  SYMS_ExtFileList list;
-  syms_memzero_struct(&list);
+  SYMS_ExtFileList list = {0};
   SYMS_ElfExtDebugRef ext_debug_ref = syms_elf_ext_debug_ref_from_elf_section_array(file, bin_accel->sections);
   if(ext_debug_ref.path.size != 0)
   {
@@ -353,8 +349,7 @@ syms_elf_ext_file_list_from_bin(SYMS_Arena *arena, SYMS_String8 file, SYMS_ElfBi
 SYMS_API SYMS_SecInfoArray
 syms_elf_sec_info_array_from_bin(SYMS_Arena *arena, SYMS_String8 data, SYMS_ElfBinAccel *bin)
 {
-  SYMS_SecInfoArray array;
-  syms_memzero_struct(&array);
+  SYMS_SecInfoArray array = {0};
   array.count = bin->sections.count;
   array.sec_info = syms_push_array_zero(arena, SYMS_SecInfo, array.count);
   for(SYMS_U64 idx = 0; idx < array.count; idx += 1)
@@ -370,6 +365,12 @@ syms_elf_default_vbase_from_bin(SYMS_ElfBinAccel *bin)
   return bin->header.base_address;
 }
 
+SYMS_API SYMS_U64
+syms_elf_entry_point_voff_from_bin(SYMS_ElfBinAccel *bin)
+{
+  return bin->header.entry_point;
+}
+
 SYMS_API SYMS_Arch
 syms_elf_arch_from_bin(SYMS_ElfBinAccel *bin){
   SYMS_Arch result = bin->header.arch;
diff --git a/Engine/Source/Runtime/SymsLib/syms/core/elf/syms_elf_parser.h b/Engine/Source/Runtime/SymsLib/syms/core/elf/syms_elf_parser.h
--- a/Engine/Source/Runtime/SymsLib/syms/core/elf/syms_elf_parser.h
+++ b/Engine/Source/Runtime/SymsLib/syms/core/elf/syms_elf_parser.h
@@ -17,6 +17,7 @@ struct SYMS_ElfImgHeader
   SYMS_U64 sh_name_low_offset;
   SYMS_U64 sh_name_high_offset;
   SYMS_U64 base_address;
+  SYMS_U64 entry_point;
 };
 
 typedef struct SYMS_ElfSection SYMS_ElfSection;
@@ -98,6 +99,7 @@ SYMS_API SYMS_ElfBinAccel *syms_elf_bin_accel_from_file(SYMS_Arena *arena, SYMS_
 SYMS_API SYMS_ExtFileList  syms_elf_ext_file_list_from_bin(SYMS_Arena *arena, SYMS_String8 file, SYMS_ElfBinAccel *bin_accel);
 SYMS_API SYMS_SecInfoArray syms_elf_sec_info_array_from_bin(SYMS_Arena *arena, SYMS_String8 data, SYMS_ElfBinAccel *bin);
 SYMS_API SYMS_U64          syms_elf_default_vbase_from_bin(SYMS_ElfBinAccel *bin);
+SYMS_API SYMS_U64          syms_elf_entry_point_voff_from_bin(SYMS_ElfBinAccel *bin);
 
 SYMS_API SYMS_Arch         syms_elf_arch_from_bin(SYMS_ElfBinAccel *bin);
 
diff --git a/Engine/Source/Runtime/SymsLib/syms/core/group/syms_group.c b/Engine/Source/Runtime/SymsLib/syms/core/group/syms_group.c
--- a/Engine/Source/Runtime/SymsLib/syms/core/group/syms_group.c
+++ b/Engine/Source/Runtime/SymsLib/syms/core/group/syms_group.c
@@ -113,14 +113,10 @@ syms_group_init(SYMS_Group *group, SYMS_ParseBundle *params){
   group->unit_proc_maps = syms_push_array_zero(group->arena, SYMS_SpatialMap1D, unit_count);
   group->unit_var_maps = syms_push_array_zero(group->arena, SYMS_SpatialMap1D, unit_count);
   group->line_sequence_maps = syms_push_array_zero(group->arena, SYMS_SpatialMap1D, unit_count);
-  group->file_to_line_to_addr_buckets = syms_push_array_zero(group->arena, SYMS_FileToLineToAddrBuckets,
-                                                             unit_count);
+  group->file_to_line_to_addr_maps = syms_push_array_zero(group->arena, SYMS_FileToLineToAddrMap,
+                                                          unit_count);
   group->unit_type_maps = syms_push_array_zero(group->arena, SYMS_SymbolNameMap, unit_count);
   
-  //- setup hash table caches
-  group->string_cons = syms_string_cons_alloc(group->arena, 4093);
-  group->file_id_2_name_map = syms_file_id_2_name_map_alloc(group->arena, 4093);
-  
   //- setup type graph
   syms_type_graph_init(&group->type_graph,
                        group->arena, &group->string_cons,
@@ -238,6 +234,11 @@ syms_group_dbg(SYMS_Group *group){
   return(group->dbg);
 }
 
+SYMS_API SYMS_Arch
+syms_group_arch(SYMS_Group *group){
+  return(group->arch);
+}
+
 SYMS_API SYMS_UnitSetAccel*
 syms_group_unit_set(SYMS_Group *group){
   return(group->unit_set);
@@ -974,37 +975,20 @@ syms_group_proc_map_from_uid(SYMS_Group *group, SYMS_UnitID uid){
   SYMS_U64 index = uid - 1;
   if (index < group->unit_count &&
       (group->unit_cache_flags[index] & SYMS_GroupUnitCacheFlag_HasProcMap) == 0){
-    //- build loose map
-    SYMS_ArenaTemp scratch = syms_get_scratch(0, 0);
-    SYMS_SpatialMap1DLoose loose = {0};
-    
+    SYMS_Arena *arena = syms_group_get_lane_arena(group);
     SYMS_String8 dbg_data = group->dbg_data;
     SYMS_DbgAccel *dbg = group->dbg;
     SYMS_UnitAccel *unit = syms_group_unit_from_uid(group, uid);
     SYMS_SymbolIDArray *sid_array = syms_group_proc_sid_array_from_uid(group, uid);
-    SYMS_U64 count = sid_array->count;
     
-    SYMS_SymbolID *sid_ptr = sid_array->ids;
-    for (SYMS_U64 i = 0; i < count; i += 1, sid_ptr += 1){
-      SYMS_SymbolID sid = *sid_ptr;
-      SYMS_SymbolKind kind = syms_group_symbol_kind_from_sid(group, unit, sid);
-      if (kind == SYMS_SymbolKind_Procedure){
-        SYMS_U64RangeArray ranges = syms_scope_vranges_from_sid(scratch.arena, dbg_data, dbg, unit, sid);
-        syms_spatial_map_1d_loose_push(scratch.arena, &loose, sid, ranges);
-      }
-    }
+    SYMS_SpatialMap1D map = syms_spatial_map_for_procs_from_sid_array(arena, dbg_data, dbg,
+                                                                      unit, sid_array);
     
-    //- bake tight map
-    SYMS_Arena *arena = syms_group_get_lane_arena(group);
-    SYMS_SpatialMap1D map = syms_spatial_map_1d_bake(arena, &loose);
+    SYMS_ASSERT_PARANOID(syms_spatial_map_1d_invariants(&map));
     
     //- save to group
     group->unit_cache_flags[index] |= SYMS_GroupUnitCacheFlag_HasProcMap;
     group->unit_proc_maps[index] = map;
-    
-    syms_release_scratch(scratch);
-    
-    SYMS_ASSERT_PARANOID(syms_spatial_map_1d_invariants(&map));
   }
   
   SYMS_SpatialMap1D *result = &syms_spatial_map_1d_nil;
@@ -1025,35 +1009,7 @@ syms_group_line_sequence_map_from_uid(SYMS_Group *group, SYMS_UnitID uid){
     //- fill spatial map array
     SYMS_Arena *arena = syms_group_get_lane_arena(group);
     SYMS_LineTable *line_table = syms_group_line_table_from_uid(group, uid);
-    SYMS_U64 count = line_table->sequence_count;
-    SYMS_SpatialMap1DRange *ranges = syms_push_array(arena, SYMS_SpatialMap1DRange, count);
-    SYMS_SpatialMap1DRange *range_ptr = ranges;
-    SYMS_U64 *seq_index_ptr = line_table->sequence_index_array;
-    SYMS_Line *line_array = line_table->line_array;
-    for (SYMS_U64 n = 1; n <= count; n += 1, range_ptr += 1){
-      // get sequence range & inc
-      SYMS_U64 first = *seq_index_ptr;
-      seq_index_ptr += 1;
-      SYMS_U64 opl = *seq_index_ptr;
-      
-      // extract sequence range
-      SYMS_U64 last = opl - 1;
-      SYMS_U64 first_voff = line_array[first].voff;
-      SYMS_U64 opl_voff = line_array[last].voff;
-      
-      // fill range ptr
-      range_ptr->range.min = first_voff;
-      range_ptr->range.max = opl_voff;
-      range_ptr->val = n;
-    }
-    
-    //- sort
-    if (!syms_spatial_map_1d_array_check_sorted(ranges, count)){
-      syms_spatial_map_1d_array_sort(ranges, count);
-    }
-    
-    //- assemble map
-    SYMS_SpatialMap1D map = {ranges, count};
+    SYMS_SpatialMap1D map = syms_spatial_map_from_line_table(arena, line_table);
     
     //- save to group
     group->unit_cache_flags[index] |= SYMS_GroupUnitCacheFlag_HasLineSeqMap;
@@ -1078,169 +1034,14 @@ syms_group_fetch_line_to_addr_maps_from_uid(SYMS_Group *group, SYMS_UnitID uid){
   SYMS_U64 index = uid - 1;
   if (index < group->unit_count &&
       (group->unit_cache_flags[index] & SYMS_GroupUnitCacheFlag_HasLineToAddrMap) == 0){
-    //- setup loose map
-    SYMS_ArenaTemp scratch = syms_get_scratch(0, 0);
-    SYMS_FileToLineToAddrLoose loose = {0};
-    
-    //- last-used cache for file nodes
-    SYMS_FileToLineToAddrLooseFile *cached_file_node = 0;
-    SYMS_U64 cached_file_id = 0;
-    
-    //- read lines
+    //- get buckets
     SYMS_LineTable *line_table = syms_group_line_table_from_uid(group, uid);
-    SYMS_U64 *seq_idx_ptr = line_table->sequence_index_array;
-    SYMS_U64 seq_count = line_table->sequence_count;
-    for (SYMS_U64 i = 0; i < seq_count; i += 1){
-      // get sequence range & inc
-      SYMS_U64 first = *seq_idx_ptr;
-      seq_idx_ptr += 1;
-      SYMS_U64 opl = *seq_idx_ptr;
-      
-      // iterate lines
-      if (first < opl){
-        SYMS_Line *line_ptr = line_table->line_array + first;
-        SYMS_U64 last = opl - 1;
-        for (SYMS_U64 j = first; j < last; j += 1, line_ptr += 1){
-          // grab line data
-          SYMS_U64 line_file_id = line_ptr->src_coord.file_id;
-          SYMS_U32 line_number = line_ptr->src_coord.line;
-          
-          // get file node
-          SYMS_FileToLineToAddrLooseFile *file_node = 0;
-          if (line_file_id == cached_file_id){
-            file_node = cached_file_node;
-          }
-          else{
-            for (SYMS_FileToLineToAddrLooseFile *node = loose.first;
-                 node != 0;
-                 node = node->next){
-              if (node->file_id == line_file_id){
-                file_node = node;
-                break;
-              }
-            }
-          }
-          if (file_node == 0){
-            file_node = syms_push_array_zero(scratch.arena, SYMS_FileToLineToAddrLooseFile, 1);
-            SYMS_QueuePush(loose.first, loose.last, file_node);
-            loose.count += 1;
-            file_node->file_id = line_file_id;
-          }
-          
-          // update the file node cache slot
-          cached_file_id = line_file_id;
-          cached_file_node = file_node;
-          
-          // get line node
-          SYMS_FileToLineToAddrLooseLine *line_node = 0;
-          for (SYMS_FileToLineToAddrLooseLine *node = file_node->first;
-               node != 0;
-               node = node->next){
-            if (node->line == line_number){
-              line_node = node;
-              break;
-            }
-          }
-          if (line_node == 0){
-            line_node = syms_push_array_zero(scratch.arena, SYMS_FileToLineToAddrLooseLine, 1);
-            SYMS_QueuePush(file_node->first, file_node->last, line_node);
-            file_node->line_count += 1;
-            line_node->line = line_number;
-          }
-          
-          // push range
-          SYMS_U64Range range = syms_make_u64_range(line_ptr->voff, (line_ptr + 1)->voff);
-          syms_u64_range_list_push(scratch.arena, &line_node->ranges, range);
-          file_node->range_count += 1;
-        }
-      }
-    }
-    
-    //- convert loose to buckets & maps
     SYMS_Arena *arena = syms_group_get_lane_arena(group);
-    
-    SYMS_FileToLineToAddrBuckets buckets = {0};
-    if (loose.count > 0){
-      buckets.bucket_count = ((loose.count + 1)*3/2) | 1;
-    }
-    buckets.buckets = syms_push_array_zero(arena, SYMS_FileToLineToAddrNode*, buckets.bucket_count);
-    for (SYMS_FileToLineToAddrLooseFile *loose_file_node = loose.first;
-         loose_file_node != 0;
-         loose_file_node = loose_file_node->next){
-      SYMS_ArenaTemp temp = syms_arena_temp_begin(scratch.arena);
-      
-      // grab counts
-      SYMS_U64 line_count = loose_file_node->line_count;
-      SYMS_U64 range_count = loose_file_node->range_count;
-      
-      // create sorted node array
-      SYMS_FileToLineToAddrLooseLine **array = syms_push_array(scratch.arena, SYMS_FileToLineToAddrLooseLine*,
-                                                               line_count);
-      {
-        SYMS_FileToLineToAddrLooseLine **line_ptr = array;
-        for (SYMS_FileToLineToAddrLooseLine *loose_line_node = loose_file_node->first;
-             loose_line_node != 0;
-             loose_line_node = loose_line_node->next, line_ptr += 1){
-          *line_ptr = loose_line_node;
-        }
-      }
-      syms_line_to_addr_line_sort(array, line_count);
-      
-      // fill line map arrays
-      SYMS_U64Range *ranges = syms_push_array(arena, SYMS_U64Range, range_count);
-      SYMS_U32 *line_range_indexes = syms_push_array(arena, SYMS_U32, line_count + 1);
-      SYMS_U32 *line_numbers = syms_push_array(arena, SYMS_U32, line_count);
-      
-      SYMS_U32 line_range_index = 0;
-      SYMS_U32 *line_range_index_ptr = line_range_indexes;
-      SYMS_U32 *line_number_ptr = line_numbers;
-      
-      {
-        SYMS_FileToLineToAddrLooseLine **line_ptr = array;
-        for (SYMS_U64 i = 0; i < line_count; i += 1, line_ptr += 1){
-          // fills
-          *line_number_ptr = (**line_ptr).line;
-          *line_range_index_ptr = line_range_index;
-          SYMS_U64Range *first_range_ptr = ranges + line_range_index;
-          SYMS_U64Range *range_ptr = first_range_ptr;
-          for (SYMS_U64RangeNode *node = (**line_ptr).ranges.first;
-               node != 0;
-               node = node->next, range_ptr += 1){
-            *range_ptr = node->range;
-          }
-          
-          // incs
-          line_number_ptr += 1;
-          line_range_index_ptr += 1;
-          line_range_index += (SYMS_U64)(range_ptr - first_range_ptr);
-        }
-        
-        // fill ender index
-        *line_range_index_ptr = line_range_index;
-      }
-      
-      // assemble the line map
-      SYMS_LineToAddrMap *new_map = syms_push_array(arena, SYMS_LineToAddrMap, 1);
-      new_map->ranges = ranges;
-      new_map->line_range_indexes = line_range_indexes;
-      new_map->line_numbers = line_numbers;
-      new_map->line_count = line_count;
-      
-      // insert bucket
-      SYMS_FileToLineToAddrNode *new_file_node = syms_push_array(arena, SYMS_FileToLineToAddrNode, 1);
-      new_file_node->file_id = loose_file_node->file_id;
-      new_file_node->map = new_map;
-      SYMS_U64 bucket_index = loose_file_node->file_id%buckets.bucket_count;
-      SYMS_StackPush(buckets.buckets[bucket_index], new_file_node);
-      
-      syms_arena_temp_end(temp);
-    }
+    SYMS_FileToLineToAddrMap map = syms_line_to_addr_map_from_line_table(arena, line_table);
     
     //- save to group
     group->unit_cache_flags[index] |= SYMS_GroupUnitCacheFlag_HasLineToAddrMap;
-    group->file_to_line_to_addr_buckets[index] = buckets;
-    
-    syms_release_scratch(scratch);
+    group->file_to_line_to_addr_maps[index] = map;
   }
   
   SYMS_ProfEnd();
@@ -1255,17 +1056,9 @@ syms_group_line_to_addr_map_from_uid_file_id(SYMS_Group *group, SYMS_UnitID uid,
   SYMS_LineToAddrMap *result = &syms_line_to_addr_map_nil;
   SYMS_U64 index = uid - 1;
   if (index < group->unit_count){
-    SYMS_FileToLineToAddrBuckets *buckets = &group->file_to_line_to_addr_buckets[index];
+    SYMS_FileToLineToAddrMap *buckets = &group->file_to_line_to_addr_maps[index];
     if (buckets->bucket_count > 0){
-      SYMS_U64 bucket_index = file_id%buckets->bucket_count;
-      for (SYMS_FileToLineToAddrNode *node = buckets->buckets[bucket_index];
-           node != 0;
-           node = node->next){
-        if (node->file_id == file_id){
-          result = node->map;
-          break;
-        }
-      }
+      result = syms_line_to_addr_map_lookup_file_id(buckets, file_id);
     }
   }
   
@@ -1292,7 +1085,7 @@ syms_group_type_map_from_uid(SYMS_Group *group, SYMS_UnitID uid){
       SYMS_UnitAccel *unit = syms_group_unit_from_uid(group, uid);
       
       // assemble loose map
-      SYMS_SymbolNameMapLoose loose = syms_symbol_name_map_begin(scratch.arena, 4001);
+      SYMS_SymbolNameMapLoose loose = {0};
       SYMS_SymbolIDArray *type_array = syms_group_type_sid_array_from_uid(group, uid);
       SYMS_SymbolID *sid_opl = type_array->ids + type_array->count;
       SYMS_SymbolID *sid_ptr = type_array->ids;
@@ -1446,125 +1239,46 @@ syms_group_vranges_from_uid_line__accelerated(SYMS_Group *group, SYMS_UnitID uid
                                               SYMS_FileID file_id, SYMS_U32 line){
   SYMS_ProfBegin("syms_group_vranges_from_uid_line__accelerated");
   SYMS_LineToAddrMap *line_map = syms_group_line_to_addr_map_from_uid_file_id(group, uid, file_id);
-  SYMS_U64 count = line_map->line_count;
-  SYMS_U32 *numbers = line_map->line_numbers;
-  SYMS_U64 index = syms_index_from_n__u32__binary_search_round_up(numbers, count, line);
+  
+  SYMS_U32 actual_line = 0;
+  SYMS_U64RangeArray ranges = syms_line_to_addr_map_lookup_nearest_line_number(line_map, line, &actual_line);
+  
   SYMS_U64RangeArray result = {0};
-  if (index < count && numbers[index] == line){
-    SYMS_U32 *range_indexes = line_map->line_range_indexes;
-    result.ranges = line_map->ranges + range_indexes[index];
-    result.count = range_indexes[index + 1] - range_indexes[index];
+  if (actual_line == line){
+    result = ranges;
   }
+  
   SYMS_ProfEnd();
   return(result);
 }
 
-//- line info binary search helper
+//- higher level mapping functions
 
-SYMS_API SYMS_U64
-syms_index_from_n__u32__binary_search_round_up(SYMS_U32 *v, SYMS_U64 count, SYMS_U32 n){
-  SYMS_ProfBegin("syms_index_from_n__u32__binary_search_round_up");
-  SYMS_U64 result = SYMS_U64_MAX;
-  if (count > 0 && n <= v[count - 1]){
-    //- binary search:
-    //   minimum index s.t. v[index] >= n
-    //  in this one we assume:
-    //   (i != j) implies (v[i] != v[j])
-    //   thus if (v[index] == n) then index already satisfies the requirement
-    SYMS_U64 first = 0;
-    SYMS_U64 opl = count;
-    for (;;){
-      SYMS_U64 mid = (first + opl - 1)/2;
-      SYMS_U64 w = v[mid];
-      if (w < n){
-        first = mid + 1;
-      }
-      else if (w > n){
-        opl = mid + 1;
-      }
-      else{
-        first = mid;
-        break;
-      }
-      if (first + 1 >= opl){
-        break;
-      }
+SYMS_API SYMS_ResolvedLine
+syms_group_resolved_location_from_proc_sid(SYMS_Group *group, SYMS_UnitAccel *unit, SYMS_SymbolID sid){
+  SYMS_ResolvedLine result = {0};
+  
+  SYMS_ArenaTemp scratch = syms_get_scratch(0, 0);
+  SYMS_U64RangeArray ranges = syms_scope_vranges_from_sid(scratch.arena, group->dbg_data, group->dbg, unit, sid);
+  if (ranges.count > 0){
+    SYMS_UnitID uid = syms_uid_from_unit(unit);
+    SYMS_U64 voff = ranges.ranges[0].min;
+    SYMS_Line line = syms_group_line_from_uid_voff__accelerated(group, uid, voff);
+    if (line.src_coord.file_id != 0){
+      SYMS_SrcCoord *src_coord = &line.src_coord;
+      SYMS_String8 file_name = syms_group_file_name_from_id(group, uid, src_coord->file_id);
+      
+      result.file_name = file_name;
+      result.line = src_coord->line;
+      result.col = src_coord->col;
+      result.voff = voff;
     }
-    result = first;
   }
-  SYMS_ProfEnd();
+  syms_release_scratch(scratch);
+  
   return(result);
 }
 
-SYMS_API void
-syms_line_to_addr_line_sort(SYMS_FileToLineToAddrLooseLine **array, SYMS_U64 count){
-  SYMS_ProfBegin("syms_line_to_addr_line_sort");
-  syms_line_to_addr_line_sort__rec(array, count);
-  SYMS_ProfEnd();
-}
-
-SYMS_API void
-syms_line_to_addr_line_sort__rec(SYMS_FileToLineToAddrLooseLine **array, SYMS_U64 count){
-  if (count > 4){
-    SYMS_U64 last = count - 1;
-    
-    // swap
-    SYMS_U64 mid = count/2;
-    SYMS_Swap(SYMS_FileToLineToAddrLooseLine*, array[mid], array[last]);
-    
-    // partition
-    SYMS_U64 key = array[last]->line;
-    SYMS_U64 j = 0;
-    for (SYMS_U64 i = 0; i < last; i += 1){
-      if (array[i]->line < key){
-        if (j != i){
-          SYMS_Swap(SYMS_FileToLineToAddrLooseLine*, array[i], array[j]);
-        }
-        j += 1;
-      }
-    }
-    
-    SYMS_Swap(SYMS_FileToLineToAddrLooseLine*, array[j], array[last]);
-    
-    // recurse
-    SYMS_U64 pivot = j;
-    syms_line_to_addr_line_sort__rec(array, pivot);
-    syms_line_to_addr_line_sort__rec(array + pivot + 1, (count - pivot - 1));
-  }
-  else if (count == 2){
-    if (array[0]->line > array[1]->line){
-      SYMS_Swap(SYMS_FileToLineToAddrLooseLine*, array[0], array[1]);
-    }
-  }
-  else if (count == 3){
-    if (array[0]->line > array[1]->line){
-      SYMS_Swap(SYMS_FileToLineToAddrLooseLine*, array[0], array[1]);
-    }
-    if (array[1]->line > array[2]->line){
-      SYMS_Swap(SYMS_FileToLineToAddrLooseLine*, array[1], array[2]);
-      if (array[0]->line > array[1]->line){
-        SYMS_Swap(SYMS_FileToLineToAddrLooseLine*, array[0], array[1]);
-      }
-    }
-  }
-  else if (count == 4){
-    if (array[0]->line > array[1]->line){
-      SYMS_Swap(SYMS_FileToLineToAddrLooseLine*, array[0], array[1]);
-    }
-    if (array[2]->line > array[3]->line){
-      SYMS_Swap(SYMS_FileToLineToAddrLooseLine*, array[2], array[3]);
-    }
-    if (array[0]->line > array[2]->line){
-      SYMS_Swap(SYMS_FileToLineToAddrLooseLine*, array[0], array[2]);
-    }
-    if (array[1]->line > array[3]->line){
-      SYMS_Swap(SYMS_FileToLineToAddrLooseLine*, array[1], array[3]);
-    }
-    if (array[1]->line > array[2]->line){
-      SYMS_Swap(SYMS_FileToLineToAddrLooseLine*, array[1], array[2]);
-    }
-  }
-}
 
 ////////////////////////////////
 // allen: Syms Group Name Mapping Functions
@@ -1633,8 +1347,7 @@ syms_group_type_graph(SYMS_Group *group){
 // allen: Syms Group Varaible Address Mapping Functions
 
 SYMS_API SYMS_SymbolID
-syms_group_var_sid_from_uid_voff__linear_scan(SYMS_TypeGraph *graph, SYMS_Group *group,
-                                              SYMS_UnitID uid, SYMS_U64 voff){
+syms_group_var_sid_from_uid_voff__linear_scan(SYMS_Group *group, SYMS_UnitID uid, SYMS_U64 voff){
   SYMS_ASSERT_PARANOID(syms_thread_lane == 0);
   SYMS_ProfBegin("syms_group_var_sid_from_uid_voff__linear_scan");
   
@@ -1657,7 +1370,7 @@ syms_group_var_sid_from_uid_voff__linear_scan(SYMS_TypeGraph *graph, SYMS_Group
       }
       else if (var_voff < voff){
         SYMS_USID var_type_usid = syms_type_from_var_sid(dbg_data, dbg, unit, sid);
-        SYMS_U64 var_size = syms_group_type_size_from_usid(graph, group, var_type_usid);
+        SYMS_U64 var_size = syms_group_type_size_from_usid(group, var_type_usid);
         if (voff < var_voff + var_size){
           hit = syms_true;
         }
@@ -1674,7 +1387,7 @@ syms_group_var_sid_from_uid_voff__linear_scan(SYMS_TypeGraph *graph, SYMS_Group
 }
 
 SYMS_API SYMS_SpatialMap1D*
-syms_group_var_map_from_uid(SYMS_TypeGraph *graph, SYMS_Group *group, SYMS_UnitID uid){
+syms_group_var_map_from_uid(SYMS_Group *group, SYMS_UnitID uid){
   SYMS_ASSERT_PARANOID(syms_thread_lane == 0);
   SYMS_ProfBegin("syms_group_var_map_from_uid");
   SYMS_U64 index = uid - 1;
@@ -1697,7 +1410,7 @@ syms_group_var_map_from_uid(SYMS_TypeGraph *graph, SYMS_Group *group, SYMS_UnitI
       if (kind == SYMS_SymbolKind_ImageRelativeVariable){
         SYMS_U64 var_virt_off = syms_voff_from_var_sid(dbg_data, dbg, unit, sid);
         SYMS_USID var_type_usid = syms_type_from_var_sid(dbg_data, dbg, unit, sid);
-        SYMS_U64 var_size = syms_group_type_size_from_usid(graph, group, var_type_usid);
+        SYMS_U64 var_size = syms_group_type_size_from_usid(group, var_type_usid);
         SYMS_U64Range range = {var_virt_off, var_virt_off + var_size};
         syms_spatial_map_1d_loose_push_single(scratch.arena, &loose, sid, range);
       }
@@ -1726,11 +1439,10 @@ syms_group_var_map_from_uid(SYMS_TypeGraph *graph, SYMS_Group *group, SYMS_UnitI
 }
 
 SYMS_API SYMS_SymbolID
-syms_group_var_sid_from_uid_voff__accelerated(SYMS_TypeGraph *graph, SYMS_Group *group,
-                                              SYMS_UnitID uid, SYMS_U64 voff){
+syms_group_var_sid_from_uid_voff__accelerated(SYMS_Group *group, SYMS_UnitID uid, SYMS_U64 voff){
   SYMS_ASSERT_PARANOID(syms_thread_lane == 0);
   SYMS_ProfBegin("syms_group_var_sid_from_uid_voff__accelerated");
-  SYMS_SpatialMap1D *map = syms_group_var_map_from_uid(graph, group, uid);
+  SYMS_SpatialMap1D *map = syms_group_var_map_from_uid(group, uid);
   SYMS_SymbolID result = syms_spatial_map_1d_value_from_point(map, voff);
   SYMS_ProfEnd();
   return(result);
@@ -1741,256 +1453,27 @@ syms_group_var_sid_from_uid_voff__accelerated(SYMS_TypeGraph *graph, SYMS_Group
 // allen: Syms Group Type Graph Functions
 
 SYMS_API SYMS_TypeNode*
-syms_group_type_from_usid(SYMS_TypeGraph *graph, SYMS_Group *group, SYMS_USID usid){
+syms_group_type_from_usid(SYMS_Group *group, SYMS_USID usid){
   SYMS_ASSERT_PARANOID(syms_thread_lane == 0);
   SYMS_ProfBegin("syms_group_type_from_usid");
-  SYMS_TypeNode *result = syms_group_type_from_usid__rec(graph, group, usid);
-  SYMS_ProfEnd();
-  return(result);
-}
-
-SYMS_API SYMS_TypeNode*
-syms_group_type_from_usid__rec(SYMS_TypeGraph *graph, SYMS_Group *group, SYMS_USID usid){
-  SYMS_TypeNode *result = 0;
   
-  //- look at unit features
-  SYMS_UnitInfo info = syms_unit_info_from_uid(group->unit_set, usid.uid);
-  SYMS_B32 can_have_types = !!(info.features&SYMS_UnitFeature_Types);
-  if (can_have_types){
-    
-    //- get cached version of type
-    SYMS_TypeNode *type_from_cache = syms_type_from_usid(&graph->type_usid_buckets, usid);
-    
-    //- construct new graph node
-    SYMS_TypeNode *new_type = 0;
-    if (type_from_cache == 0){
-      SYMS_UnitAccel *unit = syms_group_unit_from_uid(group, usid.uid);
-      SYMS_String8 dbg_data = group->dbg_data;
-      SYMS_DbgAccel *dbg = group->dbg;
-      
-      // read basic info
-      SYMS_ArenaTemp scratch = syms_get_scratch(0, 0);
-      
-      // is this a type?
-      SYMS_SymbolKind kind = syms_group_symbol_kind_from_sid(group, unit, usid.sid);
-      if (kind != SYMS_SymbolKind_Type){
-        syms_type_usid_buckets_insert(graph->arena, &graph->type_usid_buckets,
-                                      usid, new_type);
-      }
-      else{
-        
-        // TODO(allen): group wrapper
-        SYMS_TypeInfo type_info = syms_type_info_from_sid(dbg_data, dbg, unit, usid.sid);
-        SYMS_B32 pre_inserted = syms_false;
-        
-        if (syms_type_kind_is_basic(type_info.kind)){
-          SYMS_String8 name = syms_group_symbol_name_from_sid(scratch.arena, group, unit, usid.sid);
-          new_type = syms_type_basic(graph, type_info.kind, type_info.reported_size, name);
-        }
-        else switch (type_info.kind){
-          case SYMS_TypeKind_Bitfield:
-          case SYMS_TypeKind_Struct:
-          case SYMS_TypeKind_Class:
-          case SYMS_TypeKind_Union:
-          case SYMS_TypeKind_Enum:
-          case SYMS_TypeKind_Typedef:
-          case SYMS_TypeKind_ForwardStruct:
-          case SYMS_TypeKind_ForwardClass:
-          case SYMS_TypeKind_ForwardUnion:
-          case SYMS_TypeKind_ForwardEnum:
-          {
-            // stub in the new type (necessary to do this before we infer more types to handle cycles)
-            new_type = syms_push_array_zero(graph->arena, SYMS_TypeNode, 1);
-            syms_type_usid_buckets_insert(graph->arena, &graph->type_usid_buckets,
-                                          usid, new_type);
-            pre_inserted = syms_true;
-            
-            // resolve references
-            SYMS_TypeNode *referenced_type = &syms_type_graph_nil;
-            if (type_info.reported_size_interp == SYMS_SizeInterpretation_ResolveForwardReference){
-              SYMS_TypeKind match_kind = syms_type_kind_main_from_fwd(type_info.kind);
-              
-              SYMS_String8 name = syms_group_symbol_name_from_sid(scratch.arena, group, unit, usid.sid);
-              SYMS_MapAndUnit *type_map = syms_group_type_map(group);
-              SYMS_USIDList matches = syms_usid_list_from_string(scratch.arena, group->dbg_data, group->dbg,
-                                                                 type_map, name);
-              
-              SYMS_USID match_usid = {0};
-              for (SYMS_USIDNode *node = matches.first;
-                   node != 0;
-                   node = node->next){
-                SYMS_UnitAccel *check_unit = syms_group_unit_from_uid(group, node->usid.uid);
-                SYMS_TypeInfo check_info = syms_type_info_from_sid(group->dbg_data, group->dbg, check_unit,
-                                                                   node->usid.sid);
-                if (check_info.kind == match_kind){
-                  match_usid = node->usid;
-                  break;
-                }
-              }
-              
-              referenced_type = syms_group_type_from_usid__rec(graph, group, match_usid);
-            }
-            
-            // resolve direct type
-            SYMS_TypeNode *direct = syms_group_type_from_usid__rec(graph, group, type_info.direct_type);
-            
-            // compute size info
-            SYMS_U64 byte_size = 0;
-            switch (type_info.reported_size_interp){
-              case SYMS_SizeInterpretation_ByteCount:
-              {
-                byte_size = type_info.reported_size;
-              }break;
-              
-              case SYMS_SizeInterpretation_Multiplier:
-              {
-                SYMS_U64 next_size = 1;
-                if (direct->byte_size != 0){
-                  next_size = direct->byte_size;
-                }
-                byte_size = type_info.reported_size*next_size;
-              }break;
-              
-              case SYMS_SizeInterpretation_ResolveForwardReference:
-              {
-                byte_size = referenced_type->byte_size;
-              }break;
-            }
-            
-            // pick direct type
-            SYMS_TypeNode *direct_type = direct;
-            if (direct_type == &syms_type_graph_nil){
-              direct_type = referenced_type;
-            }
-            
-            // setup src coord
-            SYMS_TypeSrcCoord *src_coord = syms_push_array_zero(graph->arena, SYMS_TypeSrcCoord, 1);
-            src_coord->usid = usid;
-            src_coord->file_id = type_info.src_coord.file_id;
-            src_coord->line = type_info.src_coord.line;
-            src_coord->col = type_info.src_coord.col;
-            
-            // get name
-            SYMS_String8 name = syms_group_symbol_name_from_sid(scratch.arena, group, unit, usid.sid);
-            
-            // fill new type
-            new_type->kind = type_info.kind;
-            new_type->name = syms_string_cons(graph->arena,
-                                              graph->string_cons,
-                                              name);
-            new_type->byte_size = byte_size;
-            new_type->src_coord = src_coord;
-            new_type->direct_type = direct_type;
-            new_type->this_type = &syms_type_graph_nil;
-          }break;
-          
-          case SYMS_TypeKind_Modifier:
-          {
-            SYMS_TypeNode *direct = syms_group_type_from_usid__rec(graph, group, type_info.direct_type);
-            new_type = syms_type_mod_from_type(graph, direct, type_info.mods);
-          }break;
-          
-          case SYMS_TypeKind_Ptr:
-          case SYMS_TypeKind_LValueReference:
-          case SYMS_TypeKind_RValueReference:
-          {
-            SYMS_TypeNode *direct = syms_group_type_from_usid__rec(graph, group, type_info.direct_type);
-            new_type = syms_type_ptr_from_type(graph, type_info.kind, direct);
-          }break;
-          
-          case SYMS_TypeKind_Array:
-          {
-            SYMS_TypeNode *direct = syms_group_type_from_usid__rec(graph, group, type_info.direct_type);
-            
-            SYMS_U64 array_count = type_info.reported_size;
-            if (type_info.reported_size_interp == SYMS_SizeInterpretation_ByteCount){
-              SYMS_U64 next_size = 1;
-              if (direct->byte_size != 0){
-                next_size = direct->byte_size;
-              }
-              array_count = type_info.reported_size/next_size;
-            }
-            
-            new_type = syms_type_array_from_type(graph, direct, array_count);
-          }break;
-          
-          case SYMS_TypeKind_Proc:
-          {
-            SYMS_TypeNode *direct = syms_group_type_from_usid__rec(graph, group, type_info.direct_type);
-            
-            // read signature
-            SYMS_SigInfo sig_info = syms_sig_info_from_type_sid(scratch.arena, dbg_data, dbg, unit, usid.sid);
-            
-            SYMS_U64 param_count = sig_info.param_type_ids.count;
-            SYMS_TypeNode **params = syms_push_array(graph->arena, SYMS_TypeNode*, param_count);
-            
-            {
-              SYMS_TypeNode **param_ptr = params;
-              SYMS_TypeNode **param_opl = params + param_count;
-              SYMS_SymbolID *sid_ptr = sig_info.param_type_ids.ids;
-              for (;param_ptr < param_opl; sid_ptr += 1, param_ptr += 1){
-                SYMS_USID param_usid = syms_make_usid(sig_info.uid, *sid_ptr);
-                *param_ptr = syms_group_type_from_usid__rec(graph, group, param_usid);
-              }
-            }
-            
-            SYMS_USID this_usid = syms_make_usid(sig_info.uid, sig_info.this_type_id);
-            SYMS_TypeNode *this_type = syms_group_type_from_usid__rec(graph, group, this_usid);
-            
-            new_type = syms_type_proc_from_type(graph, direct, this_type, params, param_count);
-          }break;
-          
-          case SYMS_TypeKind_MemberPtr:
-          {
-            SYMS_TypeNode *direct = syms_group_type_from_usid__rec(graph, group, type_info.direct_type);
-            SYMS_TypeNode *container = syms_group_type_from_usid__rec(graph, group, type_info.containing_type);
-            new_type = syms_type_member_ptr_from_type(graph, container, direct);
-            if (type_info.reported_size_interp == SYMS_SizeInterpretation_ByteCount){
-              new_type->byte_size = type_info.reported_size;
-            }
-          }break;
-          
-          case SYMS_TypeKind_Variadic:
-          {
-            // TODO(allen): ?
-          }break;
-          
-          case SYMS_TypeKind_Label:
-          {
-            // TODO(allen): ?
-          }break;
-        }
-        
-        //- split out modifiers into stand alone nodes
-        if (type_info.mods != 0 && type_info.kind != SYMS_TypeKind_Modifier){
-          SYMS_TypeNode *direct = new_type;
-          new_type = syms_type_mod_from_type(graph, direct, type_info.mods);
-        }
-        
-        if (!pre_inserted){
-          syms_type_usid_buckets_insert(graph->arena, &graph->type_usid_buckets, usid, new_type);
-        }
-      }
-      
-      syms_release_scratch(scratch);
-    }
-    
-    //- set result
-    result = type_from_cache;
-    if (result == 0){
-      result = new_type;
-    }
-  }
+  SYMS_TypeGraph *graph = &group->type_graph;
   
-  if (result == 0){
-    result = &syms_type_graph_nil;
-  }
+  SYMS_TypeParseParams params = {0};
+  params.data     = group->dbg_data;
+  params.dbg      = group->dbg;
+  params.unit_set = group->unit_set;
+  params.unit     = syms_group_unit_from_uid(group, usid.uid);
+  params.uid      = usid.uid;
+  params.type_map = syms_group_type_map(group);
   
+  SYMS_TypeNode *result = syms_type_from_dbg_sid(graph, &params, usid.sid);
+  SYMS_ProfEnd();
   return(result);
 }
 
 SYMS_API SYMS_U64
-syms_group_type_size_from_usid(SYMS_TypeGraph *graph, SYMS_Group *group, SYMS_USID usid){
+syms_group_type_size_from_usid(SYMS_Group *group, SYMS_USID usid){
   SYMS_ASSERT_PARANOID(syms_thread_lane == 0);
   SYMS_ProfBegin("syms_group_type_size_from_usid");
   SYMS_UnitAccel *unit = syms_group_unit_from_uid(group, usid.uid);
@@ -2003,12 +1486,12 @@ syms_group_type_size_from_usid(SYMS_TypeGraph *graph, SYMS_Group *group, SYMS_US
     }break;
     case SYMS_SizeInterpretation_Multiplier:
     {
-      SYMS_U64 element_size = syms_group_type_size_from_usid(graph, group, type_info.direct_type);
+      SYMS_U64 element_size = syms_group_type_size_from_usid(group, type_info.direct_type);
       result = element_size*type_info.reported_size;
     }break;
     default:
     {
-      SYMS_TypeNode *type_node = syms_group_type_from_usid(graph, group, usid);
+      SYMS_TypeNode *type_node = syms_group_type_from_usid(group, usid);
       result = type_node->byte_size;
     }break;
   }
@@ -2016,11 +1499,6 @@ syms_group_type_size_from_usid(SYMS_TypeGraph *graph, SYMS_Group *group, SYMS_US
   return(result);
 }
 
-SYMS_API SYMS_TypeChain syms_group_artificial_types_from_name(SYMS_Group *group, SYMS_String8 name){
-  SYMS_TypeChain result = syms_type_from_name(&group->type_graph, name);
-  return(result);
-}
-
 SYMS_API SYMS_USIDList
 syms_group_type_list_from_name_accelerated(SYMS_Arena *arena, SYMS_Group *group, SYMS_String8 name){
   SYMS_USIDList result = {0};
@@ -2041,144 +1519,78 @@ syms_group_types_from_unit_name(SYMS_Group *group, SYMS_UnitID uid, SYMS_String8
   return(result);
 }
 
-SYMS_API SYMS_TypeMemberArray*
-syms_type_members_from_type(SYMS_Group *group,
-                            SYMS_TypeNode *type){
-  SYMS_ProfBegin("syms_type_members_from_type");
+SYMS_API SYMS_TypeMemberArray
+syms_group_type_members_from_type(SYMS_Group *group, SYMS_TypeNode *type){
+  SYMS_ProfBegin("syms_group_type_members_from_type");
   
-  SYMS_TypeMemberArray *result = &syms_type_member_array_nil;
+  SYMS_B32 has_members = syms_type_kind_is_record(type->kind);
   
-  SYMS_B32 has_members = (type->kind == SYMS_TypeKind_Class ||
-                          type->kind == SYMS_TypeKind_Struct ||
-                          type->kind == SYMS_TypeKind_Union);
+  // ensure members equipped
   if (has_members){
-    // fill cache
-    if (type->lazy_ptr == 0 && type->src_coord != 0){
-      SYMS_TypeGraph *graph = syms_group_type_graph(group);
-      
-      SYMS_USID usid = type->src_coord->usid;
-      SYMS_String8 dbg_data = group->dbg_data;
-      SYMS_DbgAccel *dbg = group->dbg;
-      
-      SYMS_UnitAccel *unit = syms_group_unit_from_uid(group, usid.uid);
-      
-      SYMS_ArenaTemp scratch = syms_get_scratch(0, 0);
-      SYMS_MemsAccel *mems_accel = syms_mems_accel_from_sid(scratch.arena, dbg_data, dbg, unit, usid.sid);
-      SYMS_U64 mem_count = syms_mem_count_from_mems(mems_accel);
-      SYMS_TypeMember *mems = syms_push_array(group->arena, SYMS_TypeMember, mem_count);
-      
-      SYMS_TypeMember *mem_ptr = mems;
-      for (SYMS_U64 n = 1; n <= mem_count; n += 1, mem_ptr += 1){
-        SYMS_ArenaTemp temp = syms_arena_temp_begin(scratch.arena);
-        
-        SYMS_MemInfo mem_info = syms_mem_info_from_number(group->arena, dbg_data, dbg, unit, mems_accel, n);
-        
-        SYMS_TypeNode *type = &syms_type_graph_nil;
-        switch (mem_info.kind){
-          case SYMS_MemKind_VTablePtr: /*no type*/ break;
-          
-          case SYMS_MemKind_DataField:
-          case SYMS_MemKind_StaticData:
-          {
-            // TODO(allen): here we lose access to the fast path from StaticData to a virtual offset.
-            // This path is only available from the DWARF backend, but there are a few ways we could
-            // preserve it.
-            SYMS_USID type_usid = syms_type_from_mem_number(dbg_data, dbg, unit, mems_accel, n);
-            type = syms_group_type_from_usid__rec(graph, group, type_usid);
-          }break;
-          
-          case SYMS_MemKind_BaseClass:
-          case SYMS_MemKind_VBaseClassPtr:
-          case SYMS_MemKind_NestedType:
-          {
-            // directly get type from member id
-            SYMS_USID type_usid = syms_type_from_mem_number(dbg_data, dbg, unit, mems_accel, n);
-            type = syms_group_type_from_usid__rec(graph, group, type_usid);
-          }break;
-          
-          case SYMS_MemKind_Method:
-          case SYMS_MemKind_StaticMethod:
-          {
-            // TODO(allen): here we lose access to the fast path a procedure symbol.
-            // This path is only available from the DWARF backend, but there are a few ways we could
-            // preserve it.
-            SYMS_SigInfo sig = syms_sig_info_from_mem_number(scratch.arena, dbg_data, dbg, unit, mems_accel, n);
-            
-            SYMS_USID return_usid = syms_make_usid(sig.uid, sig.return_type_id);
-            SYMS_USID this_usid = syms_make_usid(sig.uid, sig.this_type_id);
-            
-            SYMS_TypeNode *ret_type  = syms_group_type_from_usid__rec(graph, group, return_usid);
-            SYMS_TypeNode *this_type = syms_group_type_from_usid__rec(graph, group, this_usid);
-            SYMS_U64 param_count = sig.param_type_ids.count;
-            SYMS_TypeNode **params = syms_push_array(group->arena, SYMS_TypeNode*, param_count);
-            
-            {
-              SYMS_TypeNode **param_ptr = params;
-              SYMS_TypeNode **param_opl = params + sig.param_type_ids.count;
-              SYMS_SymbolID *param_id_ptr = sig.param_type_ids.ids;
-              for (; param_ptr < param_opl; param_id_ptr += 1, param_ptr += 1){
-                *param_ptr = syms_group_type_from_usid__rec(graph, group,
-                                                            syms_make_usid(sig.uid, *param_id_ptr));
-              }
-            }
-            
-            type = syms_type_proc_from_type(graph, ret_type, this_type, params, param_count);
-          }break;
-        }
-        
-        mem_ptr->kind = mem_info.kind;
-        mem_ptr->visibility = mem_info.visibility;
-        mem_ptr->flags = mem_info.flags;
-        mem_ptr->name = syms_string_cons(group->arena, &group->string_cons, mem_info.name);
-        mem_ptr->off = mem_info.off;
-        mem_ptr->virtual_off = mem_info.virtual_off;
-        mem_ptr->type = type;
-        
-        syms_arena_temp_end(temp);
-      }
-      
-      syms_release_scratch(scratch);
-      
-      
-      SYMS_TypeMemberArray *type_member_array = syms_push_array(group->arena, SYMS_TypeMemberArray, 1);
-      type_member_array->mems = mems;
-      type_member_array->count = mem_count;
-      type->lazy_ptr = type_member_array;
+    SYMS_UnitID uid = 0;
+    if (type->unique != 0){
+      uid = type->unique->usid.uid;
     }
     
-    // fill result
-    result = (SYMS_TypeMemberArray*)type->lazy_ptr;
+    if (uid != 0){
+      SYMS_TypeParseParams params = {0};
+      params.data = group->dbg_data;
+      params.dbg = group->dbg;
+      params.unit_set = group->unit_set;
+      params.unit = syms_group_unit_from_uid(group, uid);
+      params.uid = uid;
+      
+      syms_type_equip_members_from_dbg(&group->type_graph, &params, type);
+    }
+  }
+  
+  // fill result
+  SYMS_TypeMemberArray result = {0};
+  if (has_members && type->lazy_ptr != 0){
+    result = *(SYMS_TypeMemberArray*)type->lazy_ptr;
   }
   
   SYMS_ProfEnd();
   return(result);
 }
 
-SYMS_API SYMS_EnumInfoArray*
-syms_type_enum_members_from_type(SYMS_Group *group,
-                                 SYMS_TypeNode *type){
-  SYMS_ProfBegin("syms_type_enum_members_from_type");
-  SYMS_EnumInfoArray *result = &syms_enum_info_array_nil;
-  if (type->kind == SYMS_TypeKind_Enum){
-    // fill cache
-    if (type->lazy_ptr == 0 && type->src_coord != 0){
-      SYMS_USID usid = type->src_coord->usid;
-      SYMS_UnitAccel *unit = syms_group_unit_from_uid(group, usid.uid);
-      SYMS_EnumInfoArray *enum_info_array = syms_push_array(group->arena, SYMS_EnumInfoArray, 1);
-      *enum_info_array = syms_enum_info_array_from_sid(group->arena, group->dbg_data, group->dbg,
-                                                       unit, usid.sid);
-      type->lazy_ptr = enum_info_array;
+SYMS_API SYMS_EnumMemberArray
+syms_group_type_enum_members_from_type(SYMS_Group *group, SYMS_TypeNode *type){
+  SYMS_ProfBegin("syms_group_type_enum_members_from_type");
+  
+  SYMS_B32 has_members = syms_type_kind_is_enum(type->kind);
+  
+  // ensure members equipped
+  if (has_members){
+    SYMS_UnitID uid = 0;
+    if (type->unique != 0){
+      uid = type->unique->usid.uid;
     }
     
-    // fill result
-    result = (SYMS_EnumInfoArray*)type->lazy_ptr;
+    if (uid != 0){
+      SYMS_TypeParseParams params = {0};
+      params.data = group->dbg_data;
+      params.dbg = group->dbg;
+      params.unit_set = group->unit_set;
+      params.unit = syms_group_unit_from_uid(group, uid);
+      params.uid = uid;
+      
+      syms_type_equip_members_from_dbg(&group->type_graph, &params, type);
+    }
   }
+  
+  // fill result
+  SYMS_EnumMemberArray result = {0};
+  if (has_members && type->lazy_ptr != 0){
+    result = *(SYMS_EnumMemberArray*)type->lazy_ptr;
+  }
+  
   SYMS_ProfEnd();
   return(result);
 }
 
 ////////////////////////////////
-//~ allen: Syms File Mapp
+//~ allen: Syms File Map
 
 SYMS_API SYMS_Name2FileIDMap*
 syms_group_file_map(SYMS_Group *group){
diff --git a/Engine/Source/Runtime/SymsLib/syms/core/group/syms_group.h b/Engine/Source/Runtime/SymsLib/syms/core/group/syms_group.h
--- a/Engine/Source/Runtime/SymsLib/syms/core/group/syms_group.h
+++ b/Engine/Source/Runtime/SymsLib/syms/core/group/syms_group.h
@@ -4,52 +4,6 @@
 #ifndef SYMS_GROUP_H
 #define SYMS_GROUP_H
 
-////////////////////////////////
-//~ allen: Syms Line Mapping Structures
-
-typedef struct SYMS_LineToAddrMap{
-  SYMS_U64Range *ranges;
-  // line_range_indexes ranges from [0,line_count] inclusive so that:
-  // for-all i in [0,line_count):
-  //   (line_range_indexes[i + 1] - line_range_indexes[i]) == # of ranges for line at index i
-  SYMS_U32 *line_range_indexes;
-  SYMS_U32 *line_numbers;
-  SYMS_U64 line_count;
-} SYMS_LineToAddrMap;
-
-typedef struct SYMS_FileToLineToAddrNode{
-  struct SYMS_FileToLineToAddrNode *next;
-  SYMS_FileID file_id;
-  SYMS_LineToAddrMap *map;
-} SYMS_FileToLineToAddrNode;
-
-typedef struct SYMS_FileToLineToAddrBuckets{
-  SYMS_FileToLineToAddrNode **buckets;
-  SYMS_U64 bucket_count;
-} SYMS_FileToLineToAddrBuckets;
-
-//- loose version
-typedef struct SYMS_FileToLineToAddrLooseLine{
-  struct SYMS_FileToLineToAddrLooseLine *next;
-  SYMS_U32 line;
-  SYMS_U64RangeList ranges;
-} SYMS_FileToLineToAddrLooseLine;
-
-typedef struct SYMS_FileToLineToAddrLooseFile{
-  struct SYMS_FileToLineToAddrLooseFile *next;
-  SYMS_FileID file_id;
-  SYMS_FileToLineToAddrLooseLine *first;
-  SYMS_FileToLineToAddrLooseLine *last;
-  SYMS_U64 line_count;
-  SYMS_U64 range_count;
-} SYMS_FileToLineToAddrLooseFile;
-
-typedef struct SYMS_FileToLineToAddrLoose{
-  SYMS_FileToLineToAddrLooseFile *first;
-  SYMS_FileToLineToAddrLooseFile *last;
-  SYMS_U64 count;
-} SYMS_FileToLineToAddrLoose;
-
 ////////////////////////////////
 //~ allen: Syms Group Types
 
@@ -65,8 +19,8 @@ enum{
   SYMS_GroupUnitCacheFlag_HasProcMap        = (1 <<  7),
   SYMS_GroupUnitCacheFlag_HasVarMap         = (1 <<  8),
   SYMS_GroupUnitCacheFlag_HasLineSeqMap     = (1 <<  9),
-  SYMS_GroupUnitCacheFlag_HasLineToAddrMap  = (1 <<  10),
-  SYMS_GroupUnitCacheFlag_HasTypeNameMap    = (1 <<  11),
+  SYMS_GroupUnitCacheFlag_HasLineToAddrMap  = (1 << 10),
+  SYMS_GroupUnitCacheFlag_HasTypeNameMap    = (1 << 11),
 };
 
 typedef struct SYMS_Group{
@@ -107,7 +61,9 @@ typedef struct SYMS_Group{
   SYMS_SpatialMap1D *unit_proc_maps;
   SYMS_SpatialMap1D *unit_var_maps;
   SYMS_SpatialMap1D *line_sequence_maps;
-  SYMS_FileToLineToAddrBuckets *file_to_line_to_addr_buckets;
+  SYMS_FileToLineToAddrMap *file_to_line_to_addr_maps;
+  
+  // TODO(allen): we probably don't need to do this as the type graph does it better
   SYMS_SymbolNameMap *unit_type_maps;
   
   //- hash tables caches
@@ -154,8 +110,6 @@ SYMS_READ_ONLY SYMS_GLOBAL SYMS_String8Array syms_string_array_nil = {0};
 SYMS_READ_ONLY SYMS_GLOBAL SYMS_LineParseOut syms_line_parse_nil = {{0}};
 SYMS_READ_ONLY SYMS_GLOBAL SYMS_SpatialMap1D syms_spatial_map_1d_nil = {0};
 SYMS_READ_ONLY SYMS_GLOBAL SYMS_LineToAddrMap syms_line_to_addr_map_nil = {0};
-SYMS_READ_ONLY SYMS_GLOBAL SYMS_TypeMemberArray syms_type_member_array_nil = {0};
-SYMS_READ_ONLY SYMS_GLOBAL SYMS_EnumInfoArray syms_enum_info_array_nil = {0};
 
 ////////////////////////////////
 //~ allen: Syms Group Setup Functions
@@ -182,6 +136,7 @@ SYMS_API SYMS_String8       syms_group_bin_data(SYMS_Group *group);
 SYMS_API SYMS_BinAccel*     syms_group_bin(SYMS_Group *group);
 SYMS_API SYMS_String8       syms_group_dbg_data(SYMS_Group *group);
 SYMS_API SYMS_DbgAccel*     syms_group_dbg(SYMS_Group *group);
+SYMS_API SYMS_Arch          syms_group_arch(SYMS_Group *group);
 SYMS_API SYMS_UnitSetAccel* syms_group_unit_set(SYMS_Group *group);
 
 SYMS_API SYMS_U64           syms_group_address_size(SYMS_Group *group);
@@ -279,12 +234,10 @@ SYMS_API SYMS_Line     syms_group_line_from_uid_voff__accelerated(SYMS_Group *gr
 SYMS_API SYMS_U64RangeArray syms_group_vranges_from_uid_line__accelerated(SYMS_Group *group, SYMS_UnitID uid,
                                                                           SYMS_FileID file_id, SYMS_U32 line);
 
-//- line info binary search helper
-SYMS_API SYMS_U64      syms_index_from_n__u32__binary_search_round_up(SYMS_U32 *v, SYMS_U64 count, SYMS_U32 n);
+//- higher level mapping functions
+SYMS_API SYMS_ResolvedLine syms_group_resolved_location_from_proc_sid(SYMS_Group *group, SYMS_UnitAccel *unit,
+                                                                      SYMS_SymbolID sid);
 
-//- line-to-addr map helpers
-SYMS_API void          syms_line_to_addr_line_sort(SYMS_FileToLineToAddrLooseLine **array, SYMS_U64 count);
-SYMS_API void          syms_line_to_addr_line_sort__rec(SYMS_FileToLineToAddrLooseLine **array, SYMS_U64 count);
 
 ////////////////////////////////
 //~ allen: Syms Group Name Mapping Functions
@@ -308,32 +261,29 @@ SYMS_API SYMS_TypeGraph* syms_group_type_graph(SYMS_Group *group);
 // cannot be parallelized like the others
 
 
-SYMS_API SYMS_SymbolID syms_group_var_sid_from_uid_voff__linear_scan(SYMS_TypeGraph *graph, SYMS_Group *group,
+SYMS_API SYMS_SymbolID syms_group_var_sid_from_uid_voff__linear_scan(SYMS_Group *group,
                                                                      SYMS_UnitID uid, SYMS_U64 voff);
 
-SYMS_API SYMS_SpatialMap1D* syms_group_var_map_from_uid(SYMS_TypeGraph *graph, SYMS_Group *group, SYMS_UnitID uid);
+SYMS_API SYMS_SpatialMap1D* syms_group_var_map_from_uid(SYMS_Group *group, SYMS_UnitID uid);
 
-SYMS_API SYMS_SymbolID syms_group_var_sid_from_uid_voff__accelerated(SYMS_TypeGraph *graph, SYMS_Group *group,
+SYMS_API SYMS_SymbolID syms_group_var_sid_from_uid_voff__accelerated(SYMS_Group *group,
                                                                      SYMS_UnitID uid, SYMS_U64 voff);
 
 
 ////////////////////////////////
 //~ allen: Syms Group Type Info Functions
 
-SYMS_API SYMS_TypeNode* syms_group_type_from_usid(SYMS_TypeGraph *graph, SYMS_Group *group, SYMS_USID usid);
-SYMS_API SYMS_TypeNode* syms_group_type_from_usid__rec(SYMS_TypeGraph *graph, SYMS_Group *group, SYMS_USID usid);
-
-SYMS_API SYMS_TypeChain syms_group_artificial_types_from_name(SYMS_Group *group, SYMS_String8 name);
+SYMS_API SYMS_TypeNode* syms_group_type_from_usid(SYMS_Group *group, SYMS_USID usid);
 
 // TODO(allen): sort with name maps?
 SYMS_API SYMS_USIDList syms_group_type_list_from_name_accelerated(SYMS_Arena *arena, SYMS_Group *group,
                                                                   SYMS_String8 name);
 SYMS_API SYMS_SymbolIDArray syms_group_types_from_unit_name(SYMS_Group *group, SYMS_UnitID uid, SYMS_String8 name);
 
-SYMS_API SYMS_U64       syms_group_type_size_from_usid(SYMS_TypeGraph *graph, SYMS_Group *group, SYMS_USID usid);
+SYMS_API SYMS_U64       syms_group_type_size_from_usid(SYMS_Group *group, SYMS_USID usid);
 
-SYMS_API SYMS_TypeMemberArray* syms_type_members_from_type(SYMS_Group *group, SYMS_TypeNode *type);
-SYMS_API SYMS_EnumInfoArray*   syms_type_enum_members_from_type(SYMS_Group *group, SYMS_TypeNode *type);
+SYMS_API SYMS_TypeMemberArray syms_group_type_members_from_type(SYMS_Group *group, SYMS_TypeNode *type);
+SYMS_API SYMS_EnumMemberArray syms_group_type_enum_members_from_type(SYMS_Group *group, SYMS_TypeNode *type);
 
 ////////////////////////////////
 //~ allen: Syms File Map
