Upgrade SymsLib to latest version

diff --git a/Engine/Source/Runtime/SymsLib/syms/core/pdb/syms_msf_parser.c b/Engine/Source/Runtime/SymsLib/syms/core/pdb/syms_msf_parser.c
--- a/Engine/Source/Runtime/SymsLib/syms/core/pdb/syms_msf_parser.c
+++ b/Engine/Source/Runtime/SymsLib/syms/core/pdb/syms_msf_parser.c
@@ -4,7 +4,7 @@
 #define SYMS_MSF_PARSER_C
 
 ////////////////////////////////
-// NOTE(allen): MSF Reader Fundamentals Without Accelerator
+//~ allen: MSF Reader Fundamentals Without Accelerator
 
 SYMS_API SYMS_MsfHeaderInfo
 syms_msf_header_info_from_data_slow(SYMS_String8 data){
@@ -50,7 +50,7 @@ syms_msf_header_info_from_data_slow(SYMS_String8 data){
 }
 
 ////////////////////////////////
-// NOTE(allen): MSF Reader Accelerator Constructor
+//~ allen: MSF Reader Accelerator Constructor
 
 SYMS_API SYMS_MsfAccel*
 syms_msf_accel_from_data(SYMS_Arena *arena, SYMS_String8 data){
@@ -189,8 +189,7 @@ syms_msf_accel_from_data(SYMS_Arena *arena, SYMS_String8 data){
       syms_based_range_read(directory, directory_range, 0, 4, &stream_count);
       
       // allocate info array
-      stream_info = syms_push_array(arena, SYMS_MsfAccelStreamInfo, stream_count);
-      syms_memset(stream_info, 0, sizeof(*stream_info)*stream_count);
+      stream_info = syms_push_array_zero(arena, SYMS_MsfAccelStreamInfo, stream_count);
       
       // setup counts, sizes, and offsets for the directory data
       SYMS_U32 block_size = header.block_size;
@@ -249,8 +248,33 @@ syms_msf_deep_copy(SYMS_Arena *arena, SYMS_MsfAccel *msf){
   return(result);
 }
 
+SYMS_API SYMS_MsfAccel*
+syms_msf_accel_dummy_from_raw_data(SYMS_Arena *arena, SYMS_String8 data){
+  // setup header data
+  SYMS_MsfHeaderInfo header = {0};
+  header.index_size = 4;
+  header.block_size = data.size;
+  header.block_count = 1;
+  header.directory_size = 0;
+  header.directory_super_map = 0;
+  
+  // setup stream info
+  SYMS_MsfAccelStreamInfo *stream_info = syms_push_array(arena, SYMS_MsfAccelStreamInfo, 1);
+  stream_info->stream_indices = syms_push_array_zero(arena, SYMS_U8, 4);
+  stream_info->size = data.size;
+  
+  // fill result
+  SYMS_MsfAccel *result = syms_push_array(arena, SYMS_MsfAccel, 1);
+  result->header = header;
+  result->stream_count = 1;
+  result->stream_info = stream_info;
+  
+  return(result);
+}
+
+
 ////////////////////////////////
-// NOTE(allen): MSF Reader Fundamentals With Accelerator
+//~ allen: MSF Reader Fundamentals With Accelerator
 
 SYMS_API SYMS_MsfHeaderInfo
 syms_msf_header_info_from_msf(SYMS_MsfAccel *msf){
@@ -279,7 +303,7 @@ SYMS_API SYMS_B32
 syms_msf_bounds_check(SYMS_MsfAccel *msf, SYMS_MsfStreamNumber sn, SYMS_U32 off){
   SYMS_B32 result = syms_false;
   SYMS_MsfStreamInfo stream_info = syms_msf_stream_info_from_sn(msf, sn);
-  if (stream_info.sn != 0 && off <= stream_info.size){
+  if (off <= stream_info.size){
     result = syms_true;
   }
   return(result);
@@ -292,14 +316,14 @@ syms_msf_read(SYMS_String8 data, SYMS_MsfAccel *msf,
   
   // stream info
   SYMS_MsfStreamInfo stream_info = syms_msf_stream_info_from_sn(msf, sn);
-  if (stream_info.sn != 0 && off + size <= stream_info.size && size > 0){
+  if (size > 0 && off + size <= stream_info.size){
     // copy block-by-block
     SYMS_U64 file_size = data.size;
     SYMS_U32 block_size = msf->header.block_size;
     SYMS_U32 size_of_block_index = msf->header.index_size;
-    SYMS_U32 block_count_in_whole_file_max = SYMS_CeilIntegerDiv(file_size, block_size);
+    SYMS_U32 block_count_in_whole_file_max = SYMS_CeilIntegerDiv(file_size, (SYMS_U64)block_size);
     SYMS_U32 block_count_in_whole_file = SYMS_ClampTop(msf->header.block_count, block_count_in_whole_file_max);
-    SYMS_U32 block_count_in_stream = SYMS_CeilIntegerDiv(stream_info.size, block_size);
+    SYMS_U32 block_count_in_stream = SYMS_CeilIntegerDiv((SYMS_U64)stream_info.size, (SYMS_U64)block_size);
     
     SYMS_U32 completed_amount = 0;
     for (;;){
@@ -347,7 +371,7 @@ SYMS_API SYMS_MsfRange
 syms_msf_range_from_sn(SYMS_MsfAccel *msf, SYMS_MsfStreamNumber sn){
   SYMS_MsfStreamInfo info = syms_msf_stream_info_from_sn(msf, sn);
   SYMS_MsfRange result = {0};
-  if (info.sn != 0){
+  if (info.size != 0){
     result.sn = info.sn;
     result.size = info.size;
   }
@@ -372,7 +396,7 @@ syms_msf_sub_range_from_off_range(SYMS_MsfRange range, SYMS_U32Range off_range){
 }
 
 ////////////////////////////////
-// NOTE(allen): MSF Reader Range Helper Functions
+//~ allen: MSF Reader Range Helper Functions
 
 SYMS_API SYMS_B32
 syms_msf_bounds_check_in_range(SYMS_MsfRange range, SYMS_U32 off){
@@ -381,8 +405,8 @@ syms_msf_bounds_check_in_range(SYMS_MsfRange range, SYMS_U32 off){
 }
 
 SYMS_API SYMS_B32
-syms_msf_read_in_range(SYMS_String8 data, SYMS_MsfAccel *msf,
-                       SYMS_MsfRange range, SYMS_U32 off, SYMS_U32 size, void *out){
+syms_msf_read_in_range(SYMS_String8 data, SYMS_MsfAccel *msf, SYMS_MsfRange range,
+                       SYMS_U32 off, SYMS_U32 size, void *out){
   SYMS_B32 result = syms_false;
   if (off + size <= range.size){
     result = syms_msf_read(data, msf, range.sn, range.off + off, size, out);
@@ -411,7 +435,7 @@ syms_msf_read_zstring_in_range(SYMS_Arena *arena, SYMS_String8 data, SYMS_MsfAcc
   SYMS_U32 max_off = range.off + range.size;
   
   SYMS_MsfStreamInfo stream_info = syms_msf_stream_info_from_sn(msf, range.sn);
-  if (stream_info.sn != 0 && off < stream_info.size && off < max_off){
+  if (off < stream_info.size && off < max_off){
     // scan block-by-block
     SYMS_U64 file_size = data.size;
     SYMS_U32 block_size = msf->header.block_size;
diff --git a/Engine/Source/Runtime/SymsLib/syms/core/pdb/syms_msf_parser.h b/Engine/Source/Runtime/SymsLib/syms/core/pdb/syms_msf_parser.h
--- a/Engine/Source/Runtime/SymsLib/syms/core/pdb/syms_msf_parser.h
+++ b/Engine/Source/Runtime/SymsLib/syms/core/pdb/syms_msf_parser.h
@@ -5,7 +5,7 @@
 #define SYMS_MSF_PARSER_H
 
 ////////////////////////////////
-// NOTE(allen): MSF Parser Helper Types
+//~ allen: MSF Parser Helper Types
 
 typedef struct SYMS_MsfHeaderInfo{
   SYMS_U32 index_size;
@@ -28,7 +28,7 @@ typedef struct SYMS_MsfRange{
 }  SYMS_MsfRange;
 
 ////////////////////////////////
-// NOTE(allen): MSF Parser Accelerator
+//~ allen: MSF Parser Accelerator
 
 typedef struct SYMS_MsfAccelStreamInfo{
   SYMS_U8 *stream_indices;
@@ -42,20 +42,20 @@ typedef struct SYMS_MsfAccel{
 } SYMS_MsfAccel;
 
 ////////////////////////////////
-// NOTE(allen): MSF Reader Fundamentals Without Accelerator
+//~ allen: MSF Reader Fundamentals Without Accelerator
 
 SYMS_API SYMS_MsfHeaderInfo syms_msf_header_info_from_data_slow(SYMS_String8 data);
 
 ////////////////////////////////
-// NOTE(allen): MSF Reader Accelerator Constructor
+//~ allen: MSF Reader Accelerator Constructor
 
 SYMS_API SYMS_MsfAccel* syms_msf_accel_from_data(SYMS_Arena *arena, SYMS_String8 data);
 SYMS_API SYMS_MsfAccel* syms_msf_deep_copy(SYMS_Arena *arena, SYMS_MsfAccel *msf);
 
-// TODO(allen): dummy msf accelerator for whole blob
+SYMS_API SYMS_MsfAccel* syms_msf_accel_dummy_from_raw_data(SYMS_Arena *arena, SYMS_String8 data);
 
 ////////////////////////////////
-// NOTE(allen): MSF Reader Fundamentals With Accelerator
+//~ allen: MSF Reader Fundamentals With Accelerator
 
 SYMS_API SYMS_MsfHeaderInfo syms_msf_header_info_from_msf(SYMS_MsfAccel *msf);
 
@@ -75,7 +75,7 @@ SYMS_API SYMS_MsfRange syms_msf_sub_range(SYMS_MsfRange range, SYMS_U32 off, SYM
 SYMS_API SYMS_MsfRange syms_msf_sub_range_from_off_range(SYMS_MsfRange range, SYMS_U32Range off_range);
 
 ////////////////////////////////
-// NOTE(allen): MSF Reader Range Helper Functions
+//~ allen: MSF Reader Range Helper Functions
 
 SYMS_API SYMS_B32 syms_msf_bounds_check_in_range(SYMS_MsfRange range, SYMS_U32 off);
 SYMS_API SYMS_B32 syms_msf_read_in_range(SYMS_String8 data, SYMS_MsfAccel *msf,
diff --git a/Engine/Source/Runtime/SymsLib/syms/core/pdb/syms_pdb.c b/Engine/Source/Runtime/SymsLib/syms/core/pdb/syms_pdb.c
--- a/Engine/Source/Runtime/SymsLib/syms/core/pdb/syms_pdb.c
+++ b/Engine/Source/Runtime/SymsLib/syms/core/pdb/syms_pdb.c
@@ -4,12 +4,7 @@
 #define SYMS_PDB_C
 
 ////////////////////////////////
-// NOTE(allen): Codeview
-
-#include "syms/core/generated/syms_meta_cv.c"
-
-////////////////////////////////
-//~ NOTE(allen): PDB Hash Function
+//~ allen: PDB Hash Function
 
 SYMS_API SYMS_U32 syms_pdb_hashV1(SYMS_String8 string);
 
diff --git a/Engine/Source/Runtime/SymsLib/syms/core/pdb/syms_pdb.h b/Engine/Source/Runtime/SymsLib/syms/core/pdb/syms_pdb.h
--- a/Engine/Source/Runtime/SymsLib/syms/core/pdb/syms_pdb.h
+++ b/Engine/Source/Runtime/SymsLib/syms/core/pdb/syms_pdb.h
@@ -17,13 +17,7 @@
 // reflects some kind of structure in the format.
 
 ////////////////////////////////
-//~ NOTE(allen): PDB's Prerequisite Formats
-
-#include "syms/core/pdb/syms_cv.h"
-#include "syms/core/pdb/syms_msf.h"
-
-////////////////////////////////
-//~ NOTE(allen): PDB Format Types
+//~ allen: PDB Format Types
 
 typedef SYMS_U32 SYMS_PdbVersion;
 enum{
@@ -70,10 +64,14 @@ typedef struct SYMS_PdbInfoHeader{
 } SYMS_PdbInfoHeader;
 
 ////////////////////////////////
-//~ NOTE(allen): PDB DBI Format Types
+//~ allen: PDB DBI Format Types
 
 //- "STRTABLE" named stream contains a string table
 
+enum{
+  SYMS_PdbDbiStrTableHeader_MAGIC = 0xEFFEEFFE
+};
+
 typedef struct SYMS_PdbDbiStrTableHeader{
   SYMS_U32 magic;
   SYMS_U32 version;
@@ -112,9 +110,12 @@ enum{
 };
 
 typedef SYMS_U16 SYMS_PdbDbiBuildNumber;
+#define SYMS_PdbDbiBuildNumberNewFormatFlag 0x8000
 #define SYMS_PdbDbiBuildNumberMinor(bn) ((bn)&0xFF)
 #define SYMS_PdbDbiBuildNumberMajor(bn) (((bn) >> 8)&0x7F)
-#define SYMS_PdbDbiBuildNumberNewFormat(bn) (!!((bn)&0x8000))
+#define SYMS_PdbDbiBuildNumberNewFormat(bn) (!!((bn)&SYMS_PdbDbiBuildNumberNewFormatFlag))
+#define SYMS_PdbDbiBuildNumber(maj, min) \
+((SYMS_PdbDbiBuildNumber)(SYMS_PdbDbiBuildNumberNewFormatFlag | ((min)&0xFF) | (((maj)&0x7F) << 16)))
 
 typedef SYMS_U16 SYMS_PdbDbiHeaderFlags;
 enum{
@@ -228,7 +229,7 @@ typedef enum SYMS_PdbCompUnitRange{
 } SYMS_PdbCompUnitRange;
 
 ////////////////////////////////
-//~ NOTE(allen): PDB TPI Format Types
+//~ allen: PDB TPI Format Types
 
 // "TPI" and "IPI" fixed streams contain type/information in these structures
 
@@ -267,7 +268,7 @@ typedef struct SYMS_PdbTpiOffHint{
 } SYMS_PdbTpiOffHint;
 
 ////////////////////////////////
-//~ NOTE(allen): PDB GSI Format Types
+//~ allen: PDB GSI Format Types
 
 // "gsi" and "psi" streams contain symbol lookup tables in these types
 
@@ -293,6 +294,7 @@ typedef struct SYMS_PdbGsiHashRecord{
   SYMS_U32 cref;
 } SYMS_PdbGsiHashRecord;
 
+// TODO(allen): ?
 // NOTE(nick): This is a crutch that helps serialize in-memory GSI buckets that use 64bit pointers for next HR.
 typedef struct SYMS_PdbGsiHrOffsetCalc{
   SYMS_U32 next;
@@ -312,7 +314,7 @@ typedef struct SYMS_PdbPsiHeader{
 } SYMS_PdbPsiHeader;
 
 ////////////////////////////////
-//~ NOTE(allen): PDB Hash Function
+//~ allen: PDB Hash Function
 
 SYMS_API SYMS_U32 syms_pdb_hashV1(SYMS_String8 string);
 
