Upgrade SymsLib to latest version

diff --git a/Engine/Source/Runtime/SymsLib/syms/core/pdb/syms_pdb_parser.h b/Engine/Source/Runtime/SymsLib/syms/core/pdb/syms_pdb_parser.h
--- a/Engine/Source/Runtime/SymsLib/syms/core/pdb/syms_pdb_parser.h
+++ b/Engine/Source/Runtime/SymsLib/syms/core/pdb/syms_pdb_parser.h
@@ -4,11 +4,6 @@
 #ifndef SYMS_PDB_PARSER_H
 #define SYMS_PDB_PARSER_H
 
-////////////////////////////////
-//~ allen: PDB's Prerequisite Formats
-
-#include "syms/core/pdb/syms_msf_parser.h"
-
 ////////////////////////////////
 //~ allen: PDB Fundamental
 
@@ -18,19 +13,6 @@ typedef struct SYMS_PdbChain{
   SYMS_U32 v;
 } SYMS_PdbChain;
 
-typedef enum SYMS_PdbSymbolIDKind{
-  SYMS_PdbSymbolIDKind_Off,
-  SYMS_PdbSymbolIDKind_Index,
-} SYMS_PdbSymbolIDKind;
-
-typedef enum SYMS_PdbFileIDKind{
-  SYMS_PdbFileIDKind_Null,
-  SYMS_PdbFileIDKind_IPIOff,
-  SYMS_PdbFileIDKind_IPIStringID,
-  SYMS_PdbFileIDKind_StrTblOff,
-  SYMS_PdbFileIDKind_C11Off,
-} SYMS_PdbFileIDKind;
-
 typedef enum SYMS_PdbPseudoUnit{
   SYMS_PdbPseudoUnit_Null,
   SYMS_PdbPseudoUnit_SYM,
@@ -39,20 +21,6 @@ typedef enum SYMS_PdbPseudoUnit{
   SYMS_PdbPseudoUnit_FIRST_COMP_UNIT,
 } SYMS_PdbPseudoUnit;
 
-////////////////////////////////
-//~ allen: CV Parsing
-
-typedef struct SYMS_CvElement{
-  SYMS_MsfRange range;
-  SYMS_U32 next_off;
-  SYMS_U16 kind;
-} SYMS_CvElement;
-
-typedef struct SYMS_CvNumeric{
-  SYMS_TypeKind kind;
-  SYMS_U8 data[32];
-} SYMS_CvNumeric;
-
 ////////////////////////////////
 //~ allen: PDB DBI Information
 
@@ -100,8 +68,8 @@ typedef struct SYMS_PdbCompUnit{
   SYMS_PdbStringIndex src_file;
   SYMS_PdbStringIndex pdb_file;
   
-  SYMS_String8 module_name;
   SYMS_String8 obj_name;
+  SYMS_String8 group_name;
 } SYMS_PdbCompUnit;
 
 typedef struct SYMS_PdbCompUnitNode{
@@ -115,70 +83,6 @@ typedef struct SYMS_PdbUnitSetAccel{
   SYMS_U64 comp_count;
 } SYMS_PdbUnitSetAccel;
 
-typedef struct SYMS_PdbStub{
-  struct SYMS_PdbStub *bucket_next;
-  struct SYMS_PdbStub *parent;
-  struct SYMS_PdbStub *sibling_next;
-  struct SYMS_PdbStub *first;
-  struct SYMS_PdbStub *last;
-  SYMS_U32 off;
-  SYMS_U32 num;
-  union{
-    SYMS_U32 num2;
-    SYMS_String8 name;
-  };
-} SYMS_PdbStub;
-
-typedef struct SYMS_PdbStubRef{
-  struct SYMS_PdbStubRef *next;
-  SYMS_PdbStub *stub;
-} SYMS_PdbStubRef;
-
-typedef struct SYMS_PdbUnitAccel{
-  SYMS_FileFormat format;
-  SYMS_B32 leaf_set;
-  SYMS_MsfStreamNumber sn;
-  SYMS_PdbStub **top_stubs;
-  SYMS_U64 top_count;
-  SYMS_U64 top_min_index;
-  SYMS_PdbStub **buckets;
-  SYMS_U64 bucket_count;
-  SYMS_U64 all_count;
-  SYMS_PdbStub *ti_stubs;
-  SYMS_U64 ti_count;
-  SYMS_UnitID uid;
-  
-  SYMS_PdbStub **proc_stubs;
-  SYMS_U64 proc_count;
-  SYMS_PdbStub **var_stubs;
-  SYMS_U64 var_count;
-  SYMS_PdbStub **tls_var_stubs;
-  SYMS_U64 tls_var_count;
-  SYMS_PdbStub **thunk_stubs;
-  SYMS_U64 thunk_count;
-  SYMS_PdbStub **pub_stubs;
-  SYMS_U64 pub_count;
-  SYMS_PdbStub **udt_stubs;
-  SYMS_U64 udt_count;
-} SYMS_PdbUnitAccel;
-
-typedef struct SYMS_PdbLeafResolve{
-  SYMS_PdbStub *stub;
-  SYMS_CvLeaf leaf_kind;
-  SYMS_MsfRange leaf_range;
-  SYMS_B32 is_leaf_id;
-} SYMS_PdbLeafResolve;
-
-////////////////////////////////
-//~ allen: PDB C13
-
-typedef struct SYMS_PdbC13SubSection{
-  struct SYMS_PdbC13SubSection *next;
-  SYMS_CvSubSectionKind kind;
-  SYMS_U32 off;
-  SYMS_U32 size;
-} SYMS_PdbC13SubSection;
-
 ////////////////////////////////
 //~ allen: PDB TPI
 
@@ -190,6 +94,7 @@ typedef struct SYMS_PdbTpiOffRange{
 
 typedef struct SYMS_PdbTpiAccel{
   SYMS_MsfStreamNumber type_sn;
+  // TODO(allen): Stop playing this TPI hints game.
   SYMS_PdbTpiOffHint *hints;
   SYMS_PdbChain **buckets;
   SYMS_U32 bucket_count;
@@ -218,7 +123,11 @@ typedef struct SYMS_PdbFileAccel{
 
 typedef struct SYMS_PdbDbgAccel{
   SYMS_FileFormat format;
+  
+  // MSF view into the data
   SYMS_MsfAccel *msf;
+  
+  // specific to true PDB file
   SYMS_PdbDbiAccel dbi;
   SYMS_PdbNamedStreamArray named;
   SYMS_PdbStrtblAccel strtbl;
@@ -227,63 +136,20 @@ typedef struct SYMS_PdbDbgAccel{
   SYMS_PdbGsiAccel gsi;
   SYMS_PdbGsiAccel psi;
   SYMS_PeGuid auth_guid;
-} SYMS_PdbDbgAccel;
-
-////////////////////////////////
-//~ allen: PDB Members
-
-typedef struct SYMS_PdbMemStubNode{
-  struct SYMS_PdbMemStubNode *next;
-  SYMS_String8 name;
-  SYMS_U32 num;
-  SYMS_U32 off;
-  SYMS_U32 off2;
-} SYMS_PdbMemStubNode;
-
-typedef struct SYMS_PdbMemStubList{
-  SYMS_PdbMemStubNode *first;
-  SYMS_PdbMemStubNode *last;
-  SYMS_U64 mem_count;
-} SYMS_PdbMemStubList;
-
-typedef struct SYMS_PdbMemsAccel{
-  SYMS_FileFormat format;
-  SYMS_String8 type_name;
-  SYMS_U64 count;
-  SYMS_PdbMemStubNode **members;
-} SYMS_PdbMemsAccel;
-
-typedef struct SYMS_EnumInfoNode{
-  struct SYMS_EnumInfoNode *next;
-  SYMS_EnumInfo enum_info;
-} SYMS_EnumInfoNode;
-
-////////////////////////////////
-//~ allen: PDB Lines
-
-typedef struct SYMS_PdbFileNode{
-  struct SYMS_PdbFileNode *next;
-  SYMS_U64 count;
-  SYMS_FileID file_ids[6];
-} SYMS_PdbFileNode;
-
-typedef struct SYMS_PdbLineSequence{
-  struct SYMS_PdbLineSequence *next;
-  SYMS_Line *lines;
-  SYMS_U64 line_count;
-} SYMS_PdbLineSequence;
-
-typedef struct SYMS_PdbLineTableLoose{
-  SYMS_PdbFileNode *first_file_node;
-  SYMS_PdbFileNode *last_file_node;
-  SYMS_U64 file_count;
   
-  SYMS_PdbLineSequence *first_seq;
-  SYMS_PdbLineSequence *last_seq;
-  SYMS_U64 seq_count;
-  SYMS_U64 line_count;
-} SYMS_PdbLineTableLoose;
-
+  // TODO(allen): simplification pass
+  
+  // arch
+  SYMS_Arch arch;
+  
+  // type uid
+  SYMS_UnitID type_uid;
+  
+  // section data
+  SYMS_MsfRange section_data_range;
+  SYMS_U64 *section_voffs;
+  SYMS_U64 section_count;
+} SYMS_PdbDbgAccel;
 
 ////////////////////////////////
 //~ allen: PDB Map Accel Type
@@ -297,6 +163,11 @@ typedef struct SYMS_PdbLinkMapAccel{
   SYMS_FileFormat format;
 } SYMS_PdbLinkMapAccel;
 
+////////////////////////////////
+//~ allen: PDB Constants
+
+#define SYMS_CV_SYMBOL_ALIGN_IN_PDB  4
+
 ////////////////////////////////
 //~ allen: PDB TPI Functions
 
@@ -314,8 +185,9 @@ SYMS_API SYMS_U32         syms_pdb_tpi_base_off(SYMS_PdbTpiAccel *tpi);
 SYMS_API SYMS_MsfRange    syms_pdb_tpi_range(SYMS_MsfAccel *msf, SYMS_PdbTpiAccel *tpi);
 
 //- tpi name lookups
-SYMS_API SYMS_USIDList syms_pdb_types_from_name(SYMS_Arena *arena, SYMS_String8 data, SYMS_PdbDbgAccel *dbg,
-                                                SYMS_PdbUnitAccel *unit, SYMS_String8 name);
+SYMS_API SYMS_USIDList syms_pdb_types_from_name(SYMS_Arena *arena, SYMS_String8 data,
+                                                SYMS_PdbDbgAccel *dbg, SYMS_CvUnitAccel *unit,
+                                                SYMS_String8 name);
 
 
 ////////////////////////////////
@@ -326,7 +198,7 @@ SYMS_API SYMS_PdbGsiAccel  syms_pdb_gsi_accel_from_range(SYMS_Arena *arena, SYMS
 SYMS_API SYMS_MsfRange     syms_pdb_gsi_part_from_psi_range(SYMS_MsfRange psi_range);
 
 SYMS_API SYMS_USIDList     syms_pdb_symbols_from_name(SYMS_Arena *arena, SYMS_String8 data, SYMS_MsfAccel *msf,
-                                                      SYMS_PdbGsiAccel *gsi, SYMS_PdbUnitAccel *unit,
+                                                      SYMS_PdbGsiAccel *gsi, SYMS_CvUnitAccel *unit,
                                                       SYMS_String8 name);
 
 
@@ -334,8 +206,10 @@ SYMS_API SYMS_USIDList     syms_pdb_symbols_from_name(SYMS_Arena *arena, SYMS_St
 //~ allen: PDB Accelerator Setup
 
 // pdb specific parsing
-SYMS_API SYMS_PdbDbiAccel syms_pdb_dbi_from_msf(SYMS_String8 data, SYMS_MsfAccel *msf);
-SYMS_API SYMS_MsfRange    syms_pdb_dbi_sub_range(SYMS_PdbDbiAccel *dbi, SYMS_PdbDbiRange n);
+SYMS_API SYMS_PdbDbiAccel syms_pdb_dbi_from_msf(SYMS_String8 data, SYMS_MsfAccel *msf,
+                                                SYMS_MsfStreamNumber sn);
+SYMS_API SYMS_MsfRange    syms_pdb_dbi_sub_range(SYMS_PdbDbiAccel *dbi, SYMS_MsfStreamNumber sn,
+                                                 SYMS_PdbDbiRange n);
 SYMS_API SYMS_MsfRange    syms_pdb_dbi_stream(SYMS_MsfAccel *msf, SYMS_PdbDbiAccel *dbi, SYMS_PdbDbiStream n);
 
 SYMS_API SYMS_PdbInfoTable        syms_pdb_parse_info(SYMS_Arena *arena, SYMS_String8 data, SYMS_MsfAccel *msf);
@@ -344,10 +218,8 @@ SYMS_API SYMS_PdbStrtblAccel      syms_pdb_dbi_parse_strtbl(SYMS_Arena *arena, S
                                                             SYMS_MsfAccel *msf, SYMS_MsfStreamNumber sn);
 
 // pdb specific api
-SYMS_API SYMS_MsfAccel* syms_pdb_msf_accel_from_dbg(SYMS_PdbDbgAccel *dbg);
-
-SYMS_API SYMS_String8 syms_pdb_strtbl_string_from_off(SYMS_Arena *arena, SYMS_String8 data, SYMS_PdbDbgAccel *dbg,
-                                                      SYMS_U32 off);
+SYMS_API SYMS_String8 syms_pdb_strtbl_string_from_off(SYMS_Arena *arena, SYMS_String8 data,
+                                                      SYMS_PdbDbgAccel *dbg, SYMS_U32 off);
 SYMS_API SYMS_String8 syms_pdb_strtbl_string_from_index(SYMS_Arena *arena, SYMS_String8 data,
                                                         SYMS_PdbDbgAccel *dbg, SYMS_PdbStringIndex n);
 
@@ -374,7 +246,9 @@ SYMS_API SYMS_ExtMatchKey syms_pdb_ext_match_key_from_dbg(SYMS_String8 data, SYM
 //~ allen: PDB Sections
 
 // pdb specific
-SYMS_API SYMS_CoffSection syms_pdb_coff_section(SYMS_String8 data, SYMS_PdbDbgAccel *dbg, SYMS_U64 n);
+SYMS_API SYMS_CoffSectionHeader syms_pdb_coff_section_header(SYMS_String8 data, SYMS_PdbDbgAccel *dbg,
+                                                             SYMS_U64 n);
+SYMS_API SYMS_U64 syms_pdb_voff_from_section_n(SYMS_PdbDbgAccel *dbg, SYMS_U64 n);
 
 // main api
 SYMS_API SYMS_SecInfoArray syms_pdb_sec_info_array_from_dbg(SYMS_Arena *arena, SYMS_String8 data,
@@ -385,129 +259,94 @@ SYMS_API SYMS_SecInfoArray syms_pdb_sec_info_array_from_dbg(SYMS_Arena *arena, S
 
 // pdb specific
 SYMS_API SYMS_PdbCompUnit* syms_pdb_comp_unit_from_id(SYMS_PdbUnitSetAccel *unit_set, SYMS_UnitID id);
-SYMS_API SYMS_MsfRange     syms_pdb_msf_range_from_comp_unit(SYMS_PdbCompUnit *unit, SYMS_PdbCompUnitRange n);
+SYMS_API SYMS_MsfRange     syms_pdb_msf_range_from_comp_unit(SYMS_PdbCompUnit *unit,
+                                                             SYMS_PdbCompUnitRange n);
 
 // main api
 SYMS_API SYMS_PdbUnitSetAccel*syms_pdb_unit_set_accel_from_dbg(SYMS_Arena *arena, SYMS_String8 data,
                                                                SYMS_PdbDbgAccel *dbg);
 SYMS_API SYMS_U64             syms_pdb_unit_count_from_set(SYMS_PdbUnitSetAccel *unit_set);
-SYMS_API SYMS_UnitInfo        syms_pdb_unit_info_from_uid(SYMS_PdbUnitSetAccel *unit_set, SYMS_UnitID uid);
-SYMS_API SYMS_UnitNames       syms_pdb_unit_names_from_uid(SYMS_Arena *arena, SYMS_PdbUnitSetAccel *unit_set,
+SYMS_API SYMS_UnitInfo        syms_pdb_unit_info_from_uid(SYMS_PdbUnitSetAccel *unit_set,
+                                                          SYMS_UnitID uid);
+SYMS_API SYMS_UnitNames       syms_pdb_unit_names_from_uid(SYMS_Arena *arena,
+                                                           SYMS_PdbUnitSetAccel *unit_set,
                                                            SYMS_UnitID uid);
 
 SYMS_API SYMS_UnitRangeArray  syms_pdb_unit_ranges_from_set(SYMS_Arena *arena,
                                                             SYMS_String8 data, SYMS_PdbDbgAccel *dbg,
                                                             SYMS_PdbUnitSetAccel *unit_set);
 
-SYMS_API SYMS_UnitFeatures    syms_pdb_unit_features_from_number(SYMS_PdbUnitSetAccel *unit_set, SYMS_U64 n);
+SYMS_API SYMS_UnitFeatures    syms_pdb_unit_features_from_number(SYMS_PdbUnitSetAccel *unit_set,
+                                                                 SYMS_U64 n);
 
 ////////////////////////////////
 //~ allen: PDB Symbol Parsing
 
-// cv parsing helpers
-SYMS_API void               syms_pdb_leaf_debug_helper(SYMS_CvLeaf kind);
-SYMS_API void               syms_pdb_sym_debug_helper(SYMS_CvSymKind kind);
-
-SYMS_API SYMS_CvElement     syms_cv_element(SYMS_String8 data, SYMS_MsfAccel *msf, SYMS_MsfRange range,
-                                            SYMS_U32 off);
-SYMS_API SYMS_U32           syms_cv_read_numeric(SYMS_String8 data, SYMS_MsfAccel *msf, SYMS_MsfRange range,
-                                                 SYMS_U32 off, SYMS_CvNumeric *out);
-SYMS_API SYMS_U32           syms_cv_u32_from_numeric(SYMS_CvNumeric num);
-
-SYMS_API SYMS_ConstInfo     syms_pdb_const_info_from_cv_numeric(SYMS_CvNumeric num);
-
-SYMS_API SYMS_PdbStubRef*   syms_pdb_alloc_ref(SYMS_Arena *arena, SYMS_PdbStubRef **free_list);
-
 // cv parse
-SYMS_API SYMS_PdbUnitAccel*syms_pdb_sym_accel_from_range(SYMS_Arena *arena, SYMS_String8 data, SYMS_MsfAccel *msf,
-                                                         SYMS_MsfRange range, SYMS_UnitID uid);
-
-SYMS_API SYMS_SymbolKind   syms_pdb_sym_symbol_kind_from_id(SYMS_String8 data, SYMS_PdbDbgAccel *dbg,
-                                                            SYMS_PdbUnitAccel *unit, SYMS_SymbolID id);
-SYMS_API SYMS_String8      syms_pdb_sym_symbol_name_from_id(SYMS_Arena *arena,
-                                                            SYMS_String8 data, SYMS_PdbDbgAccel *dbg,
-                                                            SYMS_PdbUnitAccel *unit, SYMS_SymbolID id);
-
-SYMS_API SYMS_PdbUnitAccel*  syms_pdb_pub_sym_accel_from_dbg(SYMS_Arena *arena, SYMS_String8 data,
-                                                             SYMS_PdbDbgAccel *dbg);
-
-SYMS_API SYMS_PdbUnitAccel*  syms_pdb_leaf_accel_from_dbg(SYMS_Arena *arena, SYMS_String8 data,
-                                                          SYMS_PdbDbgAccel *dbg);
-SYMS_API SYMS_PdbLeafResolve syms_pdb_leaf_resolve_from_id(SYMS_String8 data, SYMS_PdbDbgAccel *dbg,
-                                                           SYMS_PdbUnitAccel *unit, SYMS_SymbolID id);
-
-SYMS_API SYMS_SymbolKind     syms_pdb_leaf_symbol_kind_from_sid(SYMS_String8 data, SYMS_PdbDbgAccel *dbg,
-                                                                SYMS_PdbUnitAccel *unit, SYMS_SymbolID id);
-SYMS_API SYMS_String8        syms_pdb_leaf_symbol_name_from_sid(SYMS_Arena *arena,
-                                                                SYMS_String8 data, SYMS_PdbDbgAccel *dbg,
-                                                                SYMS_PdbUnitAccel *unit, SYMS_SymbolID id);
+SYMS_API SYMS_CvUnitAccel*  syms_pdb_pub_sym_accel_from_dbg(SYMS_Arena *arena, SYMS_String8 data,
+                                                            SYMS_PdbDbgAccel *dbg);
 
-SYMS_API SYMS_TypeInfo       syms_pdb_leaf_type_info_from_id(SYMS_String8 data, SYMS_PdbDbgAccel *dbg,
-                                                             SYMS_PdbUnitAccel *unit, SYMS_SymbolID id);
-SYMS_API SYMS_ConstInfo      syms_pdb_leaf_const_info_from_id(SYMS_String8 data, SYMS_PdbDbgAccel *dbg,
-                                                              SYMS_PdbUnitAccel *unit, SYMS_SymbolID id);
+SYMS_API SYMS_CvUnitAccel*  syms_pdb_leaf_accel_from_dbg(SYMS_Arena *arena, SYMS_String8 data,
+                                                         SYMS_PdbDbgAccel *dbg, SYMS_UnitID uid);
 
-// pdb unit helpers
-SYMS_API SYMS_PdbStub* syms_pdb_stub_from_unit_off(SYMS_PdbUnitAccel *unit, SYMS_U32 off);
-SYMS_API SYMS_PdbStub* syms_pdb_stub_from_unit_index(SYMS_PdbUnitAccel *unit, SYMS_U32 index);
+SYMS_API SYMS_TypeInfo      syms_pdb_leaf_type_info_from_id(SYMS_String8 data, SYMS_PdbDbgAccel *dbg,
+                                                            SYMS_CvUnitAccel *unit, SYMS_SymbolID id);
+SYMS_API SYMS_ConstInfo     syms_pdb_leaf_const_info_from_id(SYMS_String8 data, SYMS_PdbDbgAccel *dbg,
+                                                             SYMS_CvUnitAccel *unit, SYMS_SymbolID id);
 
 // main api
-SYMS_API SYMS_PdbUnitAccel*syms_pdb_unit_accel_from_id(SYMS_Arena *arena, SYMS_String8 data, SYMS_PdbDbgAccel *dbg,
-                                                       SYMS_PdbUnitSetAccel *unit_set, SYMS_UnitID id);
-
-SYMS_API SYMS_UnitID       syms_pdb_uid_from_accel(SYMS_PdbUnitAccel *unit);
+SYMS_API SYMS_CvUnitAccel*syms_pdb_unit_accel_from_uid(SYMS_Arena *arena, SYMS_String8 data,
+                                                       SYMS_PdbDbgAccel *dbg, SYMS_PdbUnitSetAccel *set,
+                                                       SYMS_UnitID uid);
 
-SYMS_API SYMS_SymbolIDArray syms_pdb_proc_sid_array_from_unit(SYMS_Arena *arena, SYMS_PdbUnitAccel *unit);
-SYMS_API SYMS_SymbolIDArray syms_pdb_var_sid_array_from_unit(SYMS_Arena *arena, SYMS_PdbUnitAccel *unit);
-
-SYMS_API SYMS_UnitID       syms_pdb_tls_var_uid_from_dbg(SYMS_PdbDbgAccel *dbg);
-SYMS_API SYMS_SymbolIDArray syms_pdb_tls_var_sid_array_from_unit(SYMS_Arena *arena, SYMS_PdbUnitAccel *unit);
-
-SYMS_API SYMS_SymbolIDArray syms_pdb_type_sid_array_from_unit(SYMS_Arena *arena, SYMS_PdbUnitAccel *unit);
-
-SYMS_API SYMS_U64          syms_pdb_symbol_count_from_unit(SYMS_PdbUnitAccel *unit);
+SYMS_API SYMS_UnitID        syms_pdb_tls_var_uid_from_dbg(SYMS_PdbDbgAccel *dbg);
+SYMS_API SYMS_SymbolIDArray syms_pdb_tls_var_sid_array_from_unit(SYMS_Arena *arena,
+                                                                 SYMS_CvUnitAccel *unit);
+SYMS_API SYMS_U64          syms_pdb_symbol_count_from_unit(SYMS_CvUnitAccel *unit);
 
 SYMS_API SYMS_SymbolKind   syms_pdb_symbol_kind_from_sid(SYMS_String8 data, SYMS_PdbDbgAccel *dbg,
-                                                         SYMS_PdbUnitAccel *unit, SYMS_SymbolID id);
-SYMS_API SYMS_String8      syms_pdb_symbol_name_from_sid(SYMS_Arena *arena,SYMS_String8 data,SYMS_PdbDbgAccel *dbg,
-                                                         SYMS_PdbUnitAccel *unit, SYMS_SymbolID id);
-
+                                                         SYMS_CvUnitAccel *unit, SYMS_SymbolID id);
 
-SYMS_API SYMS_TypeInfo     syms_pdb_type_info_from_id(SYMS_String8 data, SYMS_PdbDbgAccel *dbg,
-                                                      SYMS_PdbUnitAccel *unit, SYMS_SymbolID id);
+SYMS_API SYMS_TypeInfo     syms_pdb_type_info_from_sid(SYMS_String8 data, SYMS_PdbDbgAccel *dbg,
+                                                       SYMS_CvUnitAccel *unit, SYMS_SymbolID sid);
 SYMS_API SYMS_ConstInfo    syms_pdb_const_info_from_id(SYMS_String8 data, SYMS_PdbDbgAccel *dbg,
-                                                       SYMS_PdbUnitAccel *unit, SYMS_SymbolID id);
+                                                       SYMS_CvUnitAccel *unit, SYMS_SymbolID id);
 
 ////////////////////////////////
 //~ allen: PDB Variable Info
 
 // cv parse
 SYMS_API SYMS_USID   syms_pdb_sym_type_from_var_id(SYMS_String8 data, SYMS_PdbDbgAccel *dbg,
-                                                   SYMS_PdbUnitAccel *unit, SYMS_SymbolID id);
+                                                   SYMS_CvUnitAccel *unit, SYMS_SymbolID id);
 
 SYMS_API SYMS_U64    syms_pdb_sym_voff_from_var_sid(SYMS_String8 data, SYMS_PdbDbgAccel *dbg,
-                                                    SYMS_PdbUnitAccel *unit, SYMS_SymbolID sid);
+                                                    SYMS_CvUnitAccel *unit, SYMS_SymbolID sid);
 
 SYMS_API SYMS_RegSection syms_pdb_reg_section_from_x86_reg(SYMS_CvReg cv_reg);
 SYMS_API SYMS_RegSection syms_pdb_reg_section_from_x64_reg(SYMS_CvReg cv_reg);
 SYMS_API SYMS_RegSection syms_pdb_reg_section_from_arch_reg(SYMS_Arch arch, SYMS_CvReg cv_reg);
 
 SYMS_API SYMS_RegSection syms_pdb_reg_section_from_framepointer(SYMS_String8 data,  SYMS_PdbDbgAccel *dbg,
-                                                                SYMS_MsfRange range, SYMS_PdbStub *framepointer_stub);
+                                                                SYMS_MsfRange range,
+                                                                SYMS_CvStub *framepointer_stub);
 
 // main api
-SYMS_API SYMS_USID   syms_pdb_type_from_var_id(SYMS_String8 data, SYMS_PdbDbgAccel *dbg, SYMS_PdbUnitAccel *unit,
-                                               SYMS_SymbolID id);
+SYMS_API SYMS_USID   syms_pdb_type_from_var_id(SYMS_String8 data, SYMS_PdbDbgAccel *dbg,
+                                               SYMS_CvUnitAccel *unit, SYMS_SymbolID id);
 
-SYMS_API SYMS_U64    syms_pdb_voff_from_var_sid(SYMS_String8 data, SYMS_PdbDbgAccel *dbg, SYMS_PdbUnitAccel *unit,
-                                                SYMS_SymbolID sid);
+SYMS_API SYMS_U64    syms_pdb_voff_from_var_sid(SYMS_String8 data, SYMS_PdbDbgAccel *dbg,
+                                                SYMS_CvUnitAccel *unit, SYMS_SymbolID sid);
 
-SYMS_API SYMS_Location syms_pdb_location_from_var_sid(SYMS_Arena *arena, SYMS_String8 data, SYMS_PdbDbgAccel *dbg,
-                                                      SYMS_PdbUnitAccel *unit, SYMS_SymbolID sid);
-SYMS_API SYMS_LocRangeArray syms_pdb_location_ranges_from_var_sid(SYMS_Arena *arena, SYMS_String8 data, SYMS_PdbDbgAccel *dbg,
-                                                                  SYMS_PdbUnitAccel *unit, SYMS_SymbolID sid);
-SYMS_API SYMS_Location syms_pdb_location_from_id(SYMS_Arena *arena, SYMS_String8 data, SYMS_PdbDbgAccel *dbg,
-                                                 SYMS_PdbUnitAccel *unit, SYMS_LocID loc_id);
+SYMS_API SYMS_Location syms_pdb_location_from_var_sid(SYMS_Arena *arena, SYMS_String8 data,
+                                                      SYMS_PdbDbgAccel *dbg, SYMS_CvUnitAccel *unit,
+                                                      SYMS_SymbolID sid);
+SYMS_API SYMS_LocRangeArray syms_pdb_location_ranges_from_var_sid(SYMS_Arena *arena, SYMS_String8 data,
+                                                                  SYMS_PdbDbgAccel *dbg,
+                                                                  SYMS_CvUnitAccel *unit,
+                                                                  SYMS_SymbolID sid);
+SYMS_API SYMS_Location syms_pdb_location_from_id(SYMS_Arena *arena, SYMS_String8 data,
+                                                 SYMS_PdbDbgAccel *dbg, SYMS_CvUnitAccel *unit,
+                                                 SYMS_LocID loc_id);
 
 
 ////////////////////////////////
@@ -515,28 +354,30 @@ SYMS_API SYMS_Location syms_pdb_location_from_id(SYMS_Arena *arena, SYMS_String8
 
 // cv parse
 SYMS_API void syms_pdb__field_list_parse(SYMS_Arena *arena, SYMS_String8 data, SYMS_MsfAccel *msf,
-                                         SYMS_PdbUnitAccel *unit, SYMS_U32 index, SYMS_PdbMemStubList *out);
+                                         SYMS_CvUnitAccel *unit, SYMS_U32 index,
+                                         SYMS_PdbMemStubList *out);
 
 // main api
-SYMS_API SYMS_PdbMemsAccel*syms_pdb_mems_accel_from_sid(SYMS_Arena *arena, SYMS_String8 data,
-                                                        SYMS_PdbDbgAccel *dbg, SYMS_PdbUnitAccel *unit,
-                                                        SYMS_SymbolID id);
+SYMS_API SYMS_CvMemsAccel*syms_pdb_mems_accel_from_sid(SYMS_Arena *arena, SYMS_String8 data,
+                                                       SYMS_PdbDbgAccel *dbg, SYMS_CvUnitAccel *unit,
+                                                       SYMS_SymbolID id);
 
-SYMS_API SYMS_U64      syms_pdb_mem_count_from_mems(SYMS_PdbMemsAccel *mems);
-SYMS_API SYMS_MemInfo  syms_pdb_mem_info_from_number(SYMS_Arena *arena, SYMS_String8 data, SYMS_PdbDbgAccel *dbg,
-                                                     SYMS_PdbUnitAccel *unit, SYMS_PdbMemsAccel *mems, SYMS_U64 n);
+SYMS_API SYMS_U64      syms_pdb_mem_count_from_mems(SYMS_CvMemsAccel *mems);
+SYMS_API SYMS_MemInfo  syms_pdb_mem_info_from_number(SYMS_Arena *arena, SYMS_String8 data,
+                                                     SYMS_PdbDbgAccel *dbg, SYMS_CvUnitAccel *unit,
+                                                     SYMS_CvMemsAccel *mems, SYMS_U64 n);
 
 SYMS_API SYMS_USID     syms_pdb_type_from_mem_number(SYMS_String8 data, SYMS_PdbDbgAccel *dbg,
-                                                     SYMS_PdbUnitAccel *unit, SYMS_PdbMemsAccel *mems,
+                                                     SYMS_CvUnitAccel *unit, SYMS_CvMemsAccel *mems,
                                                      SYMS_U64 n);
 
 SYMS_API SYMS_SigInfo  syms_pdb_sig_info_from_mem_number(SYMS_Arena *arena, SYMS_String8 data,
-                                                         SYMS_PdbDbgAccel *dbg, SYMS_PdbUnitAccel *unit,
-                                                         SYMS_PdbMemsAccel *mems, SYMS_U64 n);
+                                                         SYMS_PdbDbgAccel *dbg, SYMS_CvUnitAccel *unit,
+                                                         SYMS_CvMemsAccel*mems, SYMS_U64 n);
 
-SYMS_API SYMS_EnumInfoArray syms_pdb_enum_info_array_from_sid(SYMS_Arena *arena, SYMS_String8 data,
-                                                              SYMS_PdbDbgAccel *dbg, SYMS_PdbUnitAccel *unit,
-                                                              SYMS_SymbolID sid);
+SYMS_API SYMS_EnumMemberArray syms_pdb_enum_member_array_from_sid(SYMS_Arena *arena, SYMS_String8 data,
+                                                                  SYMS_PdbDbgAccel *dbg,
+                                                                  SYMS_CvUnitAccel *unit, SYMS_SymbolID sid);
 
 ////////////////////////////////
 //~ allen: PDB Procedure Info
@@ -544,16 +385,16 @@ SYMS_API SYMS_EnumInfoArray syms_pdb_enum_info_array_from_sid(SYMS_Arena *arena,
 // main api
 
 SYMS_API SYMS_UnitIDAndSig syms_pdb_proc_sig_handle_from_id(SYMS_String8 data, SYMS_PdbDbgAccel *dbg,
-                                                            SYMS_PdbUnitAccel *unit, SYMS_SymbolID id);
+                                                            SYMS_CvUnitAccel *unit, SYMS_SymbolID id);
 SYMS_API SYMS_SigInfo      syms_pdb_sig_info_from_handle(SYMS_Arena *arena, SYMS_String8 data,
-                                                         SYMS_PdbDbgAccel *dbg, SYMS_PdbUnitAccel *unit,
+                                                         SYMS_PdbDbgAccel *dbg, SYMS_CvUnitAccel *unit,
                                                          SYMS_SigHandle handle);
 
 SYMS_API SYMS_U64RangeArray syms_pdb_scope_vranges_from_sid(SYMS_Arena *arena, SYMS_String8 data,
-                                                            SYMS_PdbDbgAccel *dbg, SYMS_PdbUnitAccel *unit,
+                                                            SYMS_PdbDbgAccel *dbg, SYMS_CvUnitAccel *unit,
                                                             SYMS_SymbolID sid);
 SYMS_API SYMS_SymbolIDArray syms_pdb_scope_children_from_sid(SYMS_Arena *arena, SYMS_String8 data,
-                                                             SYMS_PdbDbgAccel *dbg, SYMS_PdbUnitAccel *unit,
+                                                             SYMS_PdbDbgAccel *dbg, SYMS_CvUnitAccel *unit,
                                                              SYMS_SymbolID id);
 
 ////////////////////////////////
@@ -561,23 +402,16 @@ SYMS_API SYMS_SymbolIDArray syms_pdb_scope_children_from_sid(SYMS_Arena *arena,
 
 // pdb specific helper
 SYMS_API SYMS_SigInfo syms_pdb_sig_info_from_sig_index(SYMS_Arena *arena, SYMS_String8 data, SYMS_PdbDbgAccel *dbg,
-                                                       SYMS_PdbUnitAccel *unit, SYMS_CvTypeIndex index);
+                                                       SYMS_CvUnitAccel *unit, SYMS_CvTypeIndex index);
 
 // main api
-SYMS_API SYMS_SigInfo syms_pdb_sig_info_from_id(SYMS_Arena *arena, SYMS_String8 data, SYMS_PdbDbgAccel *dbg,
-                                                SYMS_PdbUnitAccel *unit, SYMS_SymbolID id);
+SYMS_API SYMS_SigInfo syms_pdb_sig_info_from_id(SYMS_Arena *arena, SYMS_String8 data,
+                                                SYMS_PdbDbgAccel *dbg, SYMS_CvUnitAccel *unit,
+                                                SYMS_SymbolID id);
 
 ////////////////////////////////
 //~ allen: PDB Line Info
 
-// pdb specific
-SYMS_API SYMS_PdbC13SubSection* syms_pdb_c13_sub_sections(SYMS_Arena *arena, SYMS_String8 data, 
-                                                          SYMS_MsfAccel *msf, SYMS_MsfRange range);
-
-SYMS_API void       syms_pdb_loose_push_file_id(SYMS_Arena *arena, SYMS_PdbLineTableLoose *loose, SYMS_FileID id);
-SYMS_API SYMS_Line* syms_pdb_loose_push_sequence(SYMS_Arena *arena, SYMS_PdbLineTableLoose *loose,
-                                                 SYMS_U64 line_count);
-
 // main api
 SYMS_API SYMS_String8 syms_pdb_file_name_from_id(SYMS_Arena *arena, SYMS_String8 data, SYMS_PdbDbgAccel *dbg,
                                                  SYMS_PdbUnitSetAccel *unit_set, SYMS_UnitID uid,
@@ -596,7 +430,7 @@ SYMS_API SYMS_PdbMapAccel* syms_pdb_unmangled_symbol_map_from_dbg(SYMS_Arena *ar
                                                                   SYMS_PdbDbgAccel *dbg);
 SYMS_API SYMS_UnitID       syms_pdb_partner_uid_from_map(SYMS_PdbMapAccel *map);
 SYMS_API SYMS_USIDList     syms_pdb_usid_list_from_string(SYMS_Arena *arena, SYMS_String8 data,
-                                                          SYMS_PdbDbgAccel *dbg, SYMS_PdbUnitAccel *unit,
+                                                          SYMS_PdbDbgAccel *dbg, SYMS_CvUnitAccel *unit,
                                                           SYMS_PdbMapAccel *map, SYMS_String8 string);
 
 ////////////////////////////////
@@ -609,9 +443,9 @@ SYMS_API SYMS_PdbLinkMapAccel* syms_pdb_link_map_from_dbg(SYMS_Arena *arena, SYM
                                                           SYMS_PdbDbgAccel *dbg);
 SYMS_API SYMS_U64              syms_pdb_voff_from_link_name(SYMS_String8 data, SYMS_PdbDbgAccel *dbg,
                                                             SYMS_PdbLinkMapAccel *map,
-                                                            SYMS_PdbUnitAccel *link_unit, SYMS_String8 name);
+                                                            SYMS_CvUnitAccel *link_unit, SYMS_String8 name);
 SYMS_API SYMS_LinkNameRecArray syms_pdb_link_name_array_from_unit(SYMS_Arena *arena, SYMS_String8 data,
-                                                                  SYMS_PdbDbgAccel *dbg, SYMS_PdbUnitAccel *unit);
+                                                                  SYMS_PdbDbgAccel *dbg, SYMS_CvUnitAccel *unit);
 
 ////////////////////////////////
 //~ allen: PDB CV -> Syms Enums and Flags
diff --git a/Engine/Source/Runtime/SymsLib/syms/core/pe/syms_pe.c b/Engine/Source/Runtime/SymsLib/syms/core/pe/syms_pe.c
--- a/Engine/Source/Runtime/SymsLib/syms/core/pe/syms_pe.c
+++ b/Engine/Source/Runtime/SymsLib/syms/core/pe/syms_pe.c
@@ -3,7 +3,6 @@
 #ifndef SYMS_PE_C
 #define SYMS_PE_C
 
-#include "syms/core/generated/syms_meta_coff.c"
 #include "syms/core/generated/syms_meta_pe.c"
 
 #endif // SYMS_PE_C
diff --git a/Engine/Source/Runtime/SymsLib/syms/core/pe/syms_pe.h b/Engine/Source/Runtime/SymsLib/syms/core/pe/syms_pe.h
--- a/Engine/Source/Runtime/SymsLib/syms/core/pe/syms_pe.h
+++ b/Engine/Source/Runtime/SymsLib/syms/core/pe/syms_pe.h
@@ -19,89 +19,25 @@
 */
 
 ////////////////////////////////
-// NOTE(allen): Generated
+//~ allen: Generated
 
-#include "syms/core/generated/syms_meta_coff.h"
 #include "syms/core/generated/syms_meta_pe.h"
 
 ////////////////////////////////
-// NOTE(allen): DOS Header
+//~ allen: DOS Header
 
 // this is the "MZ" as a 16-bit short
 #define SYMS_DOS_MAGIC 0x5a4d
 
 ////////////////////////////////
-
-#define SYMS_COFF_MIN_BIG_OBJ_VERSION 2
-#define SYMS_COFF_BIG_OBJ_MAGIC_LO 0xC7A1BAD1EEBAA94Bull
-#define SYMS_COFF_BIG_OBJ_MAGIC_HI 0xAF20FAF66AA4DCB8ull
-typedef struct SYMS_CoffHeaderBigObj{
-  SYMS_U16 sig1; // SYMS_CoffMachineType_UNKNOWN
-  SYMS_U16 sig2; // SYMS_U16_MAX
-  SYMS_U16 version;
-  SYMS_U16 machine;
-  SYMS_U32 time_stamp;
-  SYMS_U8 magic[16];
-  SYMS_U32 unused[4];
-  SYMS_U32 section_count;
-  SYMS_U32 pointer_to_symbol_table;
-  SYMS_U32 number_of_symbols;
-} SYMS_CoffHeaderBigObj;
-
-#pragma pack(push, 1)
-
-typedef union SYMS_CoffSymbolName{
-  SYMS_U8 short_name[8];
-  struct{
-    // if this field is filled with zeroes we have a long name,
-    // which means name is stored in the string table 
-    // and we need to use the offset to look it up...
-    SYMS_U32 zeroes;
-    SYMS_U32 string_table_offset;
-  } long_name;
-} SYMS_CoffSymbolName;
-
-typedef struct SYMS_CoffSymbol16{
-  SYMS_CoffSymbolName name;
-  SYMS_U32 value;
-  SYMS_U16 section_number;
-  union{
-    struct{
-      SYMS_CoffSymDType msb;
-      SYMS_CoffSymType  lsb;
-    } u;
-    SYMS_U16 v;
-  } type;
-  SYMS_CoffSymStorageClass storage_class;
-  SYMS_U8 number_of_aux_symbols;
-} SYMS_CoffSymbol16;
-
-typedef struct SYMS_CoffSymbol32{
-  SYMS_CoffSymbolName name;
-  SYMS_U32 value;
-  SYMS_U32 section_number;
-  union{
-    struct{
-      SYMS_CoffSymDType msb;
-      SYMS_CoffSymType  lsb;
-    } u;
-    SYMS_U16 v;
-  } type;
-  SYMS_CoffSymStorageClass storage_class;
-  SYMS_U8 number_of_aux_symbols;
-} SYMS_CoffSymbol32;
-
-#pragma pack(pop)
-
-////////////////////////////////
-// NOTE(allen): PE Magic Numbers
+//~ allen: PE Magic Numbers
 
 #define SYMS_PE_MAGIC   0x00004550u
 #define SYMS_PE32_MAGIC     0x010bu
 #define SYMS_PE32PLUS_MAGIC 0x020bu
 
 ////////////////////////////////
-// NOTE(allen): .pdata
+//~ allen: .pdata
 
 // TODO(allen): renaming pass here.
 typedef struct SYMS_PeMipsPdata{
@@ -114,7 +50,7 @@ typedef struct SYMS_PeMipsPdata{
 
 typedef struct SYMS_PeArmPdata{
   SYMS_U32 voff_first;
-  // NOTE(allen): 
+  // NOTE(allen):
   // bits    | meaning
   // [0:7]   | prolog length
   // [8:29]  | function length
@@ -130,7 +66,7 @@ typedef struct SYMS_PeIntelPdata{
 } SYMS_PeIntelPdata;
 
 ////////////////////////////////
-// NOTE(allen): Codeview Info
+//~ allen: Codeview Info
 
 typedef struct SYMS_PeGuid{
   SYMS_U32 data1;
@@ -161,31 +97,44 @@ typedef struct SYMS_PeCvHeaderPDB70{
 #pragma pack(pop)
 
 ////////////////////////////////
-// Import
+//~ nick: Imports & Exports
 
-typedef struct SYMS_PeImportDirectoryEntry{
-  SYMS_U32 lookup_table_virt_off;
+typedef struct SYMS_PeImportEntry{
+  SYMS_U32 lookup_table_voff;
   SYMS_U32 timestamp;
   SYMS_U32 forwarder_chain;
-  SYMS_U32 name_virt_off;
-  SYMS_U32 import_addr_table_virt_off;
-} SYMS_PeImportDirectoryEntry;
-
-////////////////////////////////
-// Export
+  SYMS_U32 name_voff;
+  SYMS_U32 import_addr_table_voff;
+} SYMS_PeImportEntry;
+
+typedef struct SYMS_PeDelayedImportEntry{
+  // According to PE/COFF spec this field is unused and should be set zero,
+  // but when I compile mule with MSVC 2019 this is set to 1.
+  SYMS_U32 attributes;
+  SYMS_U32 name_voff;                       // Name of the DLL
+  SYMS_U32 module_handle_voff;              // Place where module handle from LoadLibrary is stored
+  SYMS_U32 iat_voff;
+  SYMS_U32 name_table_voff;                 // Array of hint/name or oridnals
+  SYMS_U32 bound_table_voff;                // (Optional) Points to an array of SYMS_PeThunkData
+  SYMS_U32 unload_table_voff;               // (Optional) Copy of iat_voff
+  //  0 not bound
+  // -1 if bound and real timestamp located in bounded import directory
+  // Otherwise time when dll was bound
+  SYMS_U32 timestamp; 
+} SYMS_PeDelayedImportEntry;
 
 typedef struct SYMS_PeExportTable{
   SYMS_U32 flags;                       // must be zero
   SYMS_U32 timestamp;                   // time and date when export table was created
   SYMS_U16 major_ver;                   // table version, user can change major and minor version
   SYMS_U16 minor_ver; 
-  SYMS_U32 name_virt_off;               // ASCII name of the dll
+  SYMS_U32 name_voff;                   // ASCII name of the dll
   SYMS_U32 ordinal_base;                // Starting oridnal number
   SYMS_U32 export_address_table_count;
   SYMS_U32 name_pointer_table_count;
-  SYMS_U32 export_address_table_virt_off;
-  SYMS_U32 name_pointer_table_virt_off;
-  SYMS_U32 ordinal_table_virt_off;
+  SYMS_U32 export_address_table_voff;
+  SYMS_U32 name_pointer_table_voff;
+  SYMS_U32 ordinal_table_voff;
 } SYMS_PeExportTable;
 
 #endif //SYMS_PE_H
diff --git a/Engine/Source/Runtime/SymsLib/syms/core/pe/syms_pe_parser.c b/Engine/Source/Runtime/SymsLib/syms/core/pe/syms_pe_parser.c
--- a/Engine/Source/Runtime/SymsLib/syms/core/pe/syms_pe_parser.c
+++ b/Engine/Source/Runtime/SymsLib/syms/core/pe/syms_pe_parser.c
@@ -4,7 +4,7 @@
 #define SYMS_PE_PARSER_C
 
 ////////////////////////////////
-// NOTE(allen): Functions
+//~ allen: Functions
 
 SYMS_API SYMS_PeFileAccel*
 syms_pe_file_accel_from_data(SYMS_Arena *arena, SYMS_String8 data){
@@ -57,19 +57,28 @@ syms_pe_bin_accel_from_file(SYMS_Arena *arena, SYMS_String8 data, SYMS_PeFileAcc
     
     // get sections
     SYMS_U64 sec_array_off = optional_range.max;
-    SYMS_U64 sec_array_raw_opl = sec_array_off + coff_header.section_count*sizeof(SYMS_CoffSection);
+    SYMS_U64 sec_array_raw_opl = sec_array_off + coff_header.section_count*sizeof(SYMS_CoffSectionHeader);
     SYMS_U64 sec_array_opl = SYMS_ClampTop(sec_array_raw_opl, data.size);
-    SYMS_U64 clamped_sec_count = (sec_array_opl - sec_array_off)/sizeof(SYMS_CoffSection);
-    SYMS_CoffSection *sections = (SYMS_CoffSection*)((SYMS_U8*)base + sec_array_off);
+    SYMS_U64 clamped_sec_count = (sec_array_opl - sec_array_off)/sizeof(SYMS_CoffSectionHeader);
+    SYMS_CoffSectionHeader *sections = (SYMS_CoffSectionHeader*)((SYMS_U8*)base + sec_array_off);
+    
+    // get symbols
+    SYMS_U64 symbol_array_off = coff_header.pointer_to_symbol_table;
+    SYMS_U64 symbol_count = coff_header.number_of_symbols;
+    
+    // get string table
+    SYMS_U64 string_table_off = symbol_array_off + sizeof(SYMS_CoffSymbol16) * symbol_count;
     
     // read ptional header
+    SYMS_U16 optional_magic = 0;
     SYMS_U64 image_base = 0;
+    SYMS_U64 entry_point = 0;
     SYMS_U32 data_dir_count = 0;
-    SYMS_U64Range *data_dirs_file = 0;
-    SYMS_U64Range *data_dirs_virt = 0;
+    SYMS_U64 virt_section_align = 0;
+    SYMS_U64 file_section_align = 0;
+    SYMS_U64Range *data_dir_franges = 0;
     if (optional_size > 0){
       // read magic number
-      SYMS_U16 optional_magic = 0;
       syms_based_range_read_struct(base, optional_range, 0, &optional_magic);
       
       // read optional
@@ -78,20 +87,26 @@ syms_pe_bin_accel_from_file(SYMS_Arena *arena, SYMS_String8 data, SYMS_PeFileAcc
       switch (optional_magic){
         case SYMS_PE32_MAGIC:
         {
-          SYMS_PeOptionalPe32 pe_optional = {0};
+          SYMS_PeOptionalHeader32 pe_optional = {0};
           syms_based_range_read_struct(base, optional_range, 0, &pe_optional);
           
           image_base = pe_optional.image_base;
+          entry_point = pe_optional.entry_point_va;
+          virt_section_align = pe_optional.section_alignment;
+          file_section_align = pe_optional.file_alignment;
           
           reported_data_dir_offset = sizeof(pe_optional);
           reported_data_dir_count = pe_optional.data_dir_count;
         }break;
         case SYMS_PE32PLUS_MAGIC:
         {
-          SYMS_PeOptionalPe32Plus pe_optional = {0};
+          SYMS_PeOptionalHeader32Plus pe_optional = {0};
           syms_based_range_read_struct(base, optional_range, 0, &pe_optional);
           
           image_base = pe_optional.image_base;
+          entry_point = pe_optional.entry_point_va;
+          virt_section_align = pe_optional.section_alignment;
+          file_section_align = pe_optional.file_alignment;
           
           reported_data_dir_offset = sizeof(pe_optional);
           reported_data_dir_count = pe_optional.data_dir_count;
@@ -101,15 +116,13 @@ syms_pe_bin_accel_from_file(SYMS_Arena *arena, SYMS_String8 data, SYMS_PeFileAcc
       SYMS_U32 data_dir_max = (optional_size - reported_data_dir_offset)/sizeof(SYMS_PeDataDirectory);
       data_dir_count = SYMS_ClampTop(reported_data_dir_count, data_dir_max);
       // convert PE directories to ranges
-      data_dirs_virt = syms_push_array(arena, SYMS_U64Range, data_dir_count);
-      data_dirs_file = syms_push_array(arena, SYMS_U64Range, data_dir_count);
+      data_dir_franges = syms_push_array(arena, SYMS_U64Range, data_dir_count);
       for (SYMS_U32 dir_idx = 0; dir_idx < data_dir_count; dir_idx += 1){
         SYMS_U64 dir_offset = optional_range.min + reported_data_dir_offset + sizeof(SYMS_PeDataDirectory)*dir_idx;
         SYMS_PeDataDirectory dir; syms_memzero_struct(&dir);
         syms_based_range_read_struct(base, range, dir_offset, &dir);
-        SYMS_U64 file_off = syms_pe_virt_off_to_file_off(sections, clamped_sec_count, dir.virt_off);
-        data_dirs_virt[dir_idx] = syms_make_u64_range(dir.virt_off, dir.virt_off + dir.virt_size);
-        data_dirs_file[dir_idx] = syms_make_u64_inrange(range, file_off, dir.virt_size);
+        SYMS_U64 file_off = syms_pe_foff_from_voff(sections, clamped_sec_count, dir.virt_off);
+        data_dir_franges[dir_idx] = syms_make_u64_inrange(range, file_off, dir.virt_size);
       }
     }
     
@@ -123,7 +136,7 @@ syms_pe_bin_accel_from_file(SYMS_Arena *arena, SYMS_String8 data, SYMS_PeFileAcc
     if (SYMS_PeDataDirectoryIndex_DEBUG < data_dir_count){
       // read debug directory
       SYMS_PeDebugDirectory dbg_data = {0};
-      syms_based_range_read_struct(base, data_dirs_file[SYMS_PeDataDirectoryIndex_DEBUG], 0, &dbg_data);
+      syms_based_range_read_struct(base, data_dir_franges[SYMS_PeDataDirectoryIndex_DEBUG], 0, &dbg_data);
       // extract external file info from codeview header
       if (dbg_data.type == SYMS_PeDebugDirectoryType_CODEVIEW){
         SYMS_U64 cv_offset = dbg_data.file_offset;
@@ -160,48 +173,35 @@ syms_pe_bin_accel_from_file(SYMS_Arena *arena, SYMS_String8 data, SYMS_PeFileAcc
     result = syms_push_array(arena, SYMS_PeBinAccel, 1);
     result->format = SYMS_FileFormat_PE;
     result->image_base = image_base;
+    result->entry_point = entry_point;
+    result->is_pe32 = (optional_magic == SYMS_PE32_MAGIC);
+    result->virt_section_align = virt_section_align;
+    result->file_section_align = file_section_align;
     result->section_array_off = sec_array_off;
     result->section_count = clamped_sec_count;
+    result->symbol_array_off = symbol_array_off;
+    result->symbol_count = symbol_count;
+    result->string_table_off = string_table_off;
     result->dbg_path_off = dbg_path_off;
     result->dbg_path_size = dbg_path_size;
     result->dbg_guid = dbg_guid;
     result->dbg_age = dbg_age;
+    result->dbg_time = dbg_time;
     result->arch = syms_arch_from_coff_machine_type(coff_header.machine);
-    // store virtual and file ranges for data directories
+    // store file ranges for data directories
+    result->data_dir_franges = data_dir_franges;
     result->data_dir_count = data_dir_count;
-    result->data_dirs_virt = data_dirs_virt;
-    result->data_dirs_file = data_dirs_file;
   }
   
   return(result);
 }
 
-SYMS_API SYMS_CoffFileAccel*
-syms_coff_file_accel_from_data(SYMS_Arena *arena, SYMS_String8 data){
-  SYMS_CoffFileAccel *result = (SYMS_CoffFileAccel *)&syms_format_nil;
-  // TODO(allen): 
-  return(result);
-}
-
-SYMS_API SYMS_CoffBinAccel*
-syms_coff_bin_accel_from_file(SYMS_Arena *arena, SYMS_String8 data, SYMS_CoffFileAccel *file){
-  SYMS_CoffBinAccel *result = (SYMS_CoffBinAccel *)&syms_format_nil;
-  // TODO(allen): 
-  return(result);
-}
-
 SYMS_API SYMS_Arch
 syms_pe_arch_from_bin(SYMS_PeBinAccel *bin){
   SYMS_Arch result = bin->arch;
   return(result);
 }
 
-SYMS_API SYMS_Arch
-syms_coff_arch_from_bin(SYMS_CoffBinAccel *bin){
-  SYMS_Arch result = bin->arch;
-  return(result);
-}
-
 SYMS_API SYMS_ExtFileList
 syms_pe_ext_file_list_from_bin(SYMS_Arena *arena, SYMS_String8 data, SYMS_PeBinAccel *bin){
   SYMS_ExtFileList result = {0};
@@ -218,56 +218,15 @@ syms_pe_ext_file_list_from_bin(SYMS_Arena *arena, SYMS_String8 data, SYMS_PeBinA
   return(result);
 }
 
-SYMS_API SYMS_CoffSection
-syms_pe_coff_section(SYMS_String8 data, SYMS_PeBinAccel *bin, SYMS_U64 n){
-  SYMS_CoffSection result = {0};
-  if (1 <= n && n <= bin->section_count){
-    SYMS_U64 off = bin->section_array_off + (n - 1)*sizeof(result);
-    syms_memmove(&result, data.str + off, sizeof(result));
-  }
-  return(result);
-}
-
-SYMS_API SYMS_SecInfoArray
-syms_pe_coff_sec_info_array_from_data(SYMS_Arena *arena, SYMS_String8 data,
-                                      SYMS_U64 sec_array_off, SYMS_U64 sec_count){
-  SYMS_SecInfoArray result = {0};
-  // TODO(allen): 
-  return(result);
-}
-
 SYMS_API SYMS_SecInfoArray
 syms_pe_sec_info_array_from_bin(SYMS_Arena *arena, SYMS_String8 data, SYMS_PeBinAccel *bin){
-  SYMS_SecInfoArray result = {0};
-  result.count = bin->section_count;
-  result.sec_info = syms_push_array_zero(arena, SYMS_SecInfo, result.count);
-  
-  SYMS_SecInfo *sec_info = result.sec_info;
-  SYMS_CoffSection *coff_sec = (SYMS_CoffSection*)(data.str + bin->section_array_off);
-  
-  for (SYMS_U64 i = 0; i < result.count; i += 1, sec_info += 1, coff_sec += 1){
-    // extract name
-    SYMS_U8 *name = coff_sec->name;
-    SYMS_U8 *name_ptr = name;
-    SYMS_U8 *name_opl = name + 8;
-    for (;name_ptr < name_opl && *name_ptr != 0; name_ptr += 1);
-    
-    // fill sec info
-    sec_info->name = syms_push_string_copy(arena, syms_str8_range(name, name_ptr));
-    sec_info->vrange.min = coff_sec->virt_off;
-    sec_info->vrange.max = coff_sec->virt_off + coff_sec->virt_size;
-    sec_info->frange.min = coff_sec->file_off;
-    sec_info->frange.max = coff_sec->file_off + coff_sec->file_size;
-  }
-  
+  SYMS_SecInfoArray result = syms_pecoff_sec_info_from_coff_sec(arena, data, bin->section_array_off, bin->section_count);
   return(result);
 }
 
-SYMS_API SYMS_SecInfoArray
-syms_coff_sec_info_array_from_bin(SYMS_Arena *arena, SYMS_String8 data,
-                                  SYMS_CoffBinAccel *bin){
-  SYMS_SecInfoArray result = {0};
-  // TODO(allen): 
+SYMS_API SYMS_U64
+syms_pe_entry_point_voff_from_bin(SYMS_PeBinAccel *bin){
+  SYMS_U64 result = bin->entry_point;
   return(result);
 }
 
@@ -275,207 +234,216 @@ SYMS_API SYMS_ImportArray
 syms_pe_imports_from_bin(SYMS_Arena *arena, SYMS_String8 data, SYMS_PeBinAccel *bin){
   SYMS_ArenaTemp scratch = syms_get_scratch(&arena, 1);
   
-  //- rjf: grab prerequisites
-  void *base = (void*)data.str;
-  SYMS_Arch arch = syms_pe_arch_from_bin(bin);
-  (void)arch;
-  SYMS_U64 arch_bit_size = syms_address_size_from_arch(bin->arch);
-  SYMS_U64 arch_byte_size = arch_bit_size / 8;
+  //- allen: get import sources
+  SYMS_PeImportDllArray imp_sources[2] = {0};
+  imp_sources[0] = syms_pe_regular_imports_from_bin(scratch.arena, data, bin);
+  imp_sources[1] = syms_pe_delayed_imports_from_bin(scratch.arena, data, bin);
   
-  //- rjf: grab ranges
-  SYMS_U64Range dir_range = bin->data_dirs_file[SYMS_PeDataDirectoryIndex_IMPORT];
-  SYMS_U64Range data_range = syms_make_u64_range(0, data.size);
+  //- allen: compute total # of imports
+  SYMS_U64 total_import_count = 0;
+  for (SYMS_U64 i = 0; i < SYMS_ARRAY_SIZE(imp_sources); i += 1){
+    SYMS_PeImportDllArray *imports = &imp_sources[i];
+    SYMS_PeImportDll *opl = imports->import_dlls + imports->count;
+    for (SYMS_PeImportDll *n = imports->import_dlls; n < opl; n += 1){
+      total_import_count += n->name_table.count;
+    }
+  }
   
-  //- rjf: parse import list
-  SYMS_PeImportDllNode *first_dll = 0;
-  SYMS_PeImportDllNode *last_dll = 0;
-  SYMS_U32 dll_count = 0;
-  {
-    for(SYMS_U64 read_offset = 0;;)
-    {
-      // rjf: read entry
-      SYMS_PeImportDirectoryEntry import_entry = {0};
-      read_offset += syms_based_range_read_struct(base, dir_range, read_offset, &import_entry);
-      
-      // rjf: break if this is the last entry
-      SYMS_B32 is_last_import_directory = (import_entry.lookup_table_virt_off == 0 &&
-                                           import_entry.timestamp == 0 &&
-                                           import_entry.forwarder_chain == 0 &&
-                                           import_entry.name_virt_off == 0 &&
-                                           import_entry.import_addr_table_virt_off == 0);
-      if(is_last_import_directory)
-      {
-        break;
-      }
-      
-      // rjf: grab offsets
-      SYMS_U64 roff__dll_name_file = syms_pe_bin_virt_off_to_file_off(data, bin, import_entry.name_virt_off);
-      SYMS_U64 roff__lookup_table = syms_pe_bin_virt_off_to_file_off(data, bin, import_entry.lookup_table_virt_off);
-      
-      // rjf: build import DLL node
-      SYMS_PeImportDllNode *dll_node = syms_push_array(scratch.arena, SYMS_PeImportDllNode, 1);
-      dll_node->next = 0;
-      dll_node->import_entry = import_entry;
-      dll_node->name = syms_based_range_read_string(base, data_range, roff__dll_name_file);
-      dll_node->first_import = 0;
-      dll_node->last_import = 0;
-      dll_node->import_count = 0;
-      dll_count += 1;
-      SYMS_QueuePush(first_dll, last_dll, dll_node);
+  //- allen: allocate output array
+  SYMS_ImportArray result = {0};
+  result.imports = syms_push_array_zero(arena, SYMS_Import, total_import_count);
+  result.count = total_import_count;
+  
+  //- allen: fill output array
+  SYMS_Import *imp = result.imports;
+  
+  // for each source...
+  for (SYMS_U64 i = 0; i < SYMS_ARRAY_SIZE(imp_sources); i += 1){
+    
+    // for each dll...
+    SYMS_PeImportDllArray *imports = &imp_sources[i];
+    SYMS_PeImportDll *opl = imports->import_dlls + imports->count;
+    for (SYMS_PeImportDll *n = imports->import_dlls; n < opl; n += 1){
       
-      // rjf: build list of all imports for this DLL
-      for(;;)
-      {
-        // rjf: parse entry
-        SYMS_U64 lookup_entry = 0;
-        roff__lookup_table += syms_based_range_read(base, data_range, roff__lookup_table, arch_byte_size, &lookup_entry);
+      // for each import name...
+      SYMS_U64 import_count = n->name_table.count;
+      SYMS_PeImportName *name_ptr = n->name_table.names;
+      SYMS_String8 library_name = syms_push_string_copy(arena, n->name);
+      for (SYMS_U64 j = 0; j < import_count; j += 1, name_ptr += 1){
         
-        // rjf: break if we're done
-        if(lookup_entry == 0)
-        {
-          break;
+        // fill imp & advance
+        imp->library_name = library_name;
+        if (name_ptr->lookup == SYMS_PeImportLookup_NameHint){
+          imp->name = name_ptr->name;
         }
-        
-        // rjf: build node
-        SYMS_PeImportNode *import_node = syms_push_array_zero(scratch.arena, SYMS_PeImportNode, 1);
-        SYMS_QueuePush(dll_node->first_import, dll_node->last_import, import_node);
-        dll_node->import_count += 1;
-        
-        // rjf: grab name/ordinal
-        SYMS_B32 is_ordinal = !!((arch_bit_size == 64) ? lookup_entry & 0x8000000000000000 : lookup_entry & 0x80000000);
-        if(is_ordinal)
-        {
-          import_node->ordinal = lookup_entry & 0xFFFF;
-        }
-        else
-        {
-          SYMS_U64 hn_file_off = syms_pe_bin_virt_off_to_file_off(data, bin, lookup_entry);
-          syms_based_range_read_struct(base, data_range, hn_file_off, &import_node->hint);
-          import_node->name = syms_based_range_read_string(base, data_range, hn_file_off + sizeof(import_node->hint));
+        else{
+          imp->ordinal = name_ptr->ordinal;
         }
+        imp += 1;
       }
-      
     }
   }
   
-  //- rjf: sum imports
-  SYMS_U64 total_import_count = 0;
-  for(SYMS_PeImportDllNode *n = first_dll; n != 0; n = n->next)
+  //- allen: finish
+  syms_release_scratch(scratch);
+  return(result);
+}
+
+SYMS_API SYMS_ExportArray
+syms_pe_exports_from_bin(SYMS_Arena *arena, SYMS_String8 data, SYMS_PeBinAccel *bin){
+  // TODO(allen): accelerate voff -> foff
+  
+  SYMS_ArenaTemp scratch = syms_get_scratch(&arena, 1);
+  
+  SYMS_U64Range dir_range  = bin->data_dir_franges[SYMS_PeDataDirectoryIndex_EXPORT];
+  SYMS_U64Range data_range = syms_make_u64_range(0, data.size);
+  void         *base       = (void*)data.str;
+  
+  SYMS_PeExportTable *header = (SYMS_PeExportTable*)(data.str + dir_range.min);
+  
+  // export table
+  SYMS_U32Array export_table = {0};
   {
-    total_import_count += n->import_count;
+    SYMS_U64 foff = syms_pe_bin_foff_from_voff(data, bin, header->export_address_table_voff);
+    export_table.u32 = (SYMS_U32*)(data.str + foff);
+    export_table.count = header->export_address_table_count;
   }
   
-  //- rjf: build+fill result
-  SYMS_ImportArray import_array = {0};
+  // name pointer table
+  SYMS_U32Array name_table = {0};
   {
-    SYMS_U64 import_write_idx = 0;
-    import_array.count = total_import_count;
-    import_array.imports = syms_push_array_zero(arena, SYMS_Import, import_array.count);
-    for(SYMS_PeImportDllNode *dll_node = first_dll;
-        dll_node != 0;
-        dll_node = dll_node->next)
-    {
-      for(SYMS_PeImportNode *import_node = dll_node->first_import;
-          import_node != 0;
-          import_node = import_node->next)
-      {
-        SYMS_Import *imp = &import_array.imports[import_write_idx];
-        imp->library_name = syms_push_string_copy(arena, dll_node->name);
-        imp->name = syms_push_string_copy(arena, import_node->name);
-        imp->ordinal = import_node->ordinal;
-        import_write_idx += 1;
-      }
+    SYMS_U64 foff = syms_pe_bin_foff_from_voff(data, bin, header->name_pointer_table_voff);
+    name_table.u32 = (SYMS_U32*)(data.str + foff);
+    name_table.count = header->name_pointer_table_count;
+  }
+  
+  // ordinal table
+  SYMS_U16Array ordinal_table = {0};
+  {
+    SYMS_U64 foff = syms_pe_bin_foff_from_voff(data, bin, header->ordinal_table_voff);
+    ordinal_table.u16 = (SYMS_U16*)(data.str + foff);
+    ordinal_table.count = header->name_pointer_table_count;
+  }
+  
+  // Scan export address table to get accruate count of ordinals. 
+  // We can't rely on "name_pointer_table_count" becuase it is possible
+  // to define an export without a name through NONAME attribute in DEF file
+  SYMS_U64 ordinal_count = 0;
+  for (SYMS_U64 voff_idx = 0; voff_idx < export_table.count; voff_idx += 1){
+    if (export_table.u32[voff_idx] != 0){
+      ordinal_count += 1;
     }
   }
   
-  //- rjf: return
-  syms_release_scratch(scratch);
-  return import_array;
-}
-
-SYMS_API SYMS_ExportArray
-syms_pe_exports_from_bin(SYMS_Arena *arena, SYMS_String8 data, SYMS_PeBinAccel *bin){
-  SYMS_U64Range dir_range = bin->data_dirs_file[SYMS_PeDataDirectoryIndex_EXPORT];
-  SYMS_U64Range data_range = syms_make_u64_range(0, data.size);
-  void *base = (void*)data.str;
+  // ordinal base
+  SYMS_U32 ordinal_base = header->ordinal_base;
+  
+  //
+  SYMS_U64 max_ordinal = header->export_address_table_count;
+  SYMS_B8 *is_ordinal_used = syms_push_array_zero(scratch.arena, SYMS_B8, max_ordinal);
   
-  SYMS_PeExportTable header; syms_memzero_struct(&header);
-  syms_based_range_read_struct(base, dir_range, 0, &header);
-  
-  SYMS_U32 export_address_table_file_off = syms_pe_bin_virt_off_to_file_off(data, bin, header.export_address_table_virt_off);
-  SYMS_U32 name_pointer_table_file_off   = syms_pe_bin_virt_off_to_file_off(data, bin, header.name_pointer_table_virt_off);
-  SYMS_U32 ordinal_table_file_off        = syms_pe_bin_virt_off_to_file_off(data, bin, header.ordinal_table_virt_off);
-  
-  SYMS_U64Range export_table_range       = syms_make_u64_inrange(data_range, export_address_table_file_off, sizeof(SYMS_U32) * header.export_address_table_count);
-  SYMS_U64Range name_pointer_table_range = syms_make_u64_inrange(data_range, name_pointer_table_file_off,   sizeof(SYMS_U32) * header.name_pointer_table_count);
-  SYMS_U64Range ordinal_table_range      = syms_make_u64_inrange(data_range, ordinal_table_file_off,        sizeof(SYMS_U16) * header.name_pointer_table_count);
-  
-  // NOTE(nick): There is a DLL name string in the export header. This is
-  // intersting because if we read image from memory we can figure out its
-  // name by reading this string.
-#if 0
-  SYMS_U32 dll_name_file_off = syms_pe_bin_virt_off_to_file_off(data, bin, header.name_virt_off);
-  SYMS_String8 dll_name = syms_based_range_read_string(base, data_range, dll_name_file_off);
-#endif
-  
-  SYMS_ExportArray export_array;
-  export_array.count = header.name_pointer_table_count;
-  export_array.exports = syms_push_array_zero(arena, SYMS_Export, export_array.count);
-  
-  for (SYMS_U32 i = 0; i < header.name_pointer_table_count; i += 1) {
-    SYMS_U32 name_virt_off = 0;
-    syms_based_range_read_struct(base, name_pointer_table_range, i * sizeof(name_virt_off), &name_virt_off);
-    SYMS_U32 name_file_off = syms_pe_bin_virt_off_to_file_off(data, bin, name_virt_off);
-    SYMS_String8 name = syms_based_range_read_string(base, data_range, name_file_off);
+  // allocate output array
+  SYMS_ExportArray result = {0};
+  result.exports = syms_push_array_zero(arena, SYMS_Export, ordinal_count);
+  result.count   = ordinal_count;
+  
+  // fill output array with named exports
+  SYMS_Export *exp = result.exports;
+  SYMS_Export *exp_opl = result.exports + result.count;
+  SYMS_U32 *name_ptr = name_table.u32;
+  SYMS_U32 *name_opl = name_table.u32 + name_table.count;
+  SYMS_U16 *ordinal_ptr = ordinal_table.u16;
+  for (; name_ptr < name_opl;
+       name_ptr += 1, ordinal_ptr += 1, exp += 1){
+    // get name
+    SYMS_String8 name = {0};
+    {
+      SYMS_U32 foff = syms_pe_bin_foff_from_voff(data, bin, *name_ptr);
+      SYMS_String8 raw_name = syms_based_range_read_string(base, data_range, foff);
+      name = syms_push_string_copy(arena, raw_name);
+    }
+
+    // get ordinal
+    SYMS_U16 ordinal_nobase = *ordinal_ptr;
+
+    // mark ordinal
+    is_ordinal_used[ordinal_nobase] = syms_true;
     
-    SYMS_U16 ordinal = 0;
-    syms_based_range_read_struct(base, ordinal_table_range, i * sizeof(ordinal), &ordinal);
-    SYMS_U32 biased_ordinal = ordinal + header.ordinal_base;
+    // get voff
+    SYMS_U32 export_voff = 0;
+    if (ordinal_nobase < export_table.count){
+      export_voff = export_table.u32[ordinal_nobase];
+    }
     
-    SYMS_U32 export_virt_off = 0;
-    syms_based_range_read_struct(base, export_table_range, ordinal * sizeof(export_virt_off), &export_virt_off);
-    SYMS_U32 export_file_off = syms_pe_bin_virt_off_to_file_off(data, bin, export_virt_off);
+    // make ordinal
+    SYMS_U16 ordinal = ordinal_nobase + ordinal_base;
+
+    // fill out export
+    exp->name = name;
+    exp->address = export_voff;
+    exp->ordinal = ordinal;
+  }
+
+  // fill output array with noname exports
+  for (SYMS_U64 ordinal_nobase = 0; ordinal_nobase < export_table.count; ordinal_nobase += 1){
+    SYMS_U32 voff = export_table.u32[ordinal_nobase];
+    SYMS_B32 is_voff_taken = (voff != 0);
+    SYMS_B32 is_ordinal_free = !is_ordinal_used[ordinal_nobase];
+    if (is_voff_taken && is_ordinal_free){
+      if (exp < exp_opl){
+        exp->name = syms_str8(0,0);
+        exp->address = voff;
+        exp->ordinal = ordinal_nobase + ordinal_base;
+        exp += 1;
+      }
+    }
+  }
+
+  // resolve forward ferences 
+  for (exp = result.exports; exp < exp_opl; exp += 1){
+    // determine if this is a forwarder
+    SYMS_U32 export_foff = syms_pe_bin_foff_from_voff(data, bin, exp->address);
+    SYMS_B32 is_forwarder = (dir_range.min <= export_foff && exp->address < dir_range.max);
     
-    SYMS_Export *exp = &export_array.exports[i];
-    SYMS_B32 is_virt_off_forwarder = bin->data_dirs_virt[SYMS_PeDataDirectoryIndex_EXPORT].min <= export_virt_off && export_virt_off < bin->data_dirs_virt[SYMS_PeDataDirectoryIndex_EXPORT].max;
-    if (is_virt_off_forwarder) {
-      SYMS_String8 library_name = syms_str8_lit("@MALFORMED_LIBRARY_NAME@");
-      SYMS_String8 import_name = syms_str8_lit("@MALFORMED_IMPORT_NAME@");
-      // parse forwarder command string that is formatted like "<DLL>.<IMPORT>"
-      SYMS_String8 forwarder = syms_based_range_read_string(base, data_range, export_file_off);
-      SYMS_ArenaTemp temp = syms_arena_temp_begin(arena);
-      SYMS_String8List list = syms_string_split(temp.arena, forwarder, '.');
-      if (list.node_count == 2) {
-        library_name = list.first->string;
-        import_name  = list.last->string;
+    if (is_forwarder){
+      // parse forwarder command string that is formatted like "<DLL>.<IMPORT>" (e.g. NTDLL.RtlZeroMemory)
+      SYMS_String8 library_name = syms_str8_lit("<error-dll>");
+      SYMS_String8 import_name = syms_str8_lit("<error-import>");
+      {
+        SYMS_String8 forwarder = syms_based_range_read_string(base, data_range, export_foff);
+        SYMS_ArenaTemp temp = syms_arena_temp_begin(arena);
+        SYMS_String8List list = syms_string_split(temp.arena, forwarder, '.');
+        if (list.node_count == 2){
+          library_name = list.first->string;
+          import_name  = list.last->string;
+        }
+        syms_arena_temp_end(temp);
       }
-      syms_arena_temp_end(temp);
-      // fill out export
-      exp->name = syms_push_string_copy(arena, name);
+      
+      // fill forwarder
       exp->address = 0;
-      exp->ordinal = biased_ordinal;
       exp->forwarder_library_name = syms_push_string_copy(arena, library_name);
       exp->forwarder_import_name = syms_push_string_copy(arena, import_name);
-    } else {
-      // fill out export
-      exp->name = syms_push_string_copy(arena, name);
-      exp->address = export_virt_off;
-      exp->ordinal = biased_ordinal;
     }
   }
   
-  return export_array;
+  syms_release_scratch(scratch);
+  return(result);
 }
 
-
 ////////////////////////////////
-// NOTE(allen): PE Specific Helpers
+//~ allen: PE Specific Helpers
+
+// pdata
 
 SYMS_API SYMS_U64
 syms_pe_binary_search_intel_pdata(SYMS_String8 data, SYMS_PeBinAccel *bin, SYMS_U64 voff){
+  // TODO(allen): cleanup pass.
+  
   SYMS_ASSERT(bin->arch == SYMS_Arch_X86 || bin->arch == SYMS_Arch_X64);
-  SYMS_U64 pdata_off = bin->data_dirs_file[SYMS_PeDataDirectoryIndex_EXCEPTIONS].min;
-  SYMS_U64 pdata_count = syms_u64_range_size(bin->data_dirs_file[SYMS_PeDataDirectoryIndex_EXCEPTIONS]) / sizeof(SYMS_PeIntelPdata);
+  SYMS_U64Range range = bin->data_dir_franges[SYMS_PeDataDirectoryIndex_EXCEPTIONS];
+  SYMS_U64 pdata_off = range.min;
+  SYMS_U64 pdata_count = (range.max - range.min)/sizeof(SYMS_PeIntelPdata);
   
   SYMS_U64 result = 0;
   // check if this bin includes a pdata array
@@ -521,11 +489,13 @@ syms_pe_binary_search_intel_pdata(SYMS_String8 data, SYMS_PeBinAccel *bin, SYMS_
   return(result);
 }
 
+// sections
+
 SYMS_API SYMS_U64
 syms_pe_sec_number_from_voff(SYMS_String8 data, SYMS_PeBinAccel *bin, SYMS_U64 voff){
   SYMS_U64 sec_count = bin->section_count;
-  SYMS_CoffSection *sec_array = (SYMS_CoffSection*)((SYMS_U8*)data.str + bin->section_array_off);
-  SYMS_CoffSection *sec_ptr = sec_array;
+  SYMS_CoffSectionHeader *sec_array = (SYMS_CoffSectionHeader*)((SYMS_U8*)data.str + bin->section_array_off);
+  SYMS_CoffSectionHeader *sec_ptr = sec_array;
   SYMS_U64 result = 0;
   for (SYMS_U64 i = 1; i <= sec_count; i += 1, sec_ptr += 1){
     if (sec_ptr->virt_off <= voff && voff < sec_ptr->virt_off + sec_ptr->virt_size){
@@ -541,8 +511,8 @@ syms_pe_ptr_from_sec_number(SYMS_String8 data, SYMS_PeBinAccel *bin, SYMS_U64 n)
   void *result = 0;
   SYMS_U64 sec_count = bin->section_count;
   if (1 <= n && n <= sec_count){
-    SYMS_CoffSection *sec_array = (SYMS_CoffSection*)((SYMS_U8*)data.str + bin->section_array_off);
-    SYMS_CoffSection *sec = sec_array + n - 1;
+    SYMS_CoffSectionHeader *sec_array = (SYMS_CoffSectionHeader*)((SYMS_U8*)data.str + bin->section_array_off);
+    SYMS_CoffSectionHeader *sec = sec_array + n - 1;
     if (sec->file_size > 0){
       result = data.str + sec->file_off;
     }
@@ -563,9 +533,9 @@ SYMS_API void*
 syms_pe_ptr_from_voff(SYMS_String8 data, SYMS_PeBinAccel *bin, SYMS_U64 voff){
   //- get the section for this voff
   SYMS_U64 sec_count = bin->section_count;
-  SYMS_CoffSection *sec_array = (SYMS_CoffSection*)((SYMS_U8*)data.str + bin->section_array_off);
-  SYMS_CoffSection *sec_ptr = sec_array;
-  SYMS_CoffSection *sec = 0;
+  SYMS_CoffSectionHeader *sec_array = (SYMS_CoffSectionHeader*)((SYMS_U8*)data.str + bin->section_array_off);
+  SYMS_CoffSectionHeader *sec_ptr = sec_array;
+  SYMS_CoffSectionHeader *sec = 0;
   for (SYMS_U64 i = 1; i <= sec_count; i += 1, sec_ptr += 1){
     if (sec_ptr->virt_off <= voff && voff < sec_ptr->virt_off + sec_ptr->virt_size){
       sec = sec_ptr;
@@ -583,20 +553,313 @@ syms_pe_ptr_from_voff(SYMS_String8 data, SYMS_PeBinAccel *bin, SYMS_U64 voff){
 }
 
 SYMS_API SYMS_U64
-syms_pe_virt_off_to_file_off(SYMS_CoffSection *sections, SYMS_U32 section_count, SYMS_U64 virt_off){
+syms_pe_foff_from_voff(SYMS_CoffSectionHeader *sections, SYMS_U32 section_count, SYMS_U64 virt_off){
   SYMS_U64 file_off = 0;
   for (SYMS_U32 sect_idx = 0; sect_idx < section_count; sect_idx += 1){
-    SYMS_CoffSection *sect = &sections[sect_idx];
+    SYMS_CoffSectionHeader *sect = &sections[sect_idx];
     if (sect->virt_off <= virt_off && virt_off < sect->virt_off + sect->virt_size){
-      file_off = sect->file_off + (virt_off - sect->virt_off);
+      if (!(sect->flags & SYMS_CoffSectionFlag_CNT_UNINITIALIZED_DATA)){
+        file_off = sect->file_off + (virt_off - sect->virt_off);
+      }
+      break;
     }
   }
   return file_off;
 }
 
 SYMS_API SYMS_U64
-syms_pe_bin_virt_off_to_file_off(SYMS_String8 data, SYMS_PeBinAccel *bin, SYMS_U64 virt_off){
-  return syms_pe_virt_off_to_file_off((SYMS_CoffSection*)((SYMS_U8*)data.str + bin->section_array_off), bin->section_count, virt_off);
+syms_pe_bin_foff_from_voff(SYMS_String8 data, SYMS_PeBinAccel *bin, SYMS_U64 virt_off){
+  return syms_pe_foff_from_voff((SYMS_CoffSectionHeader*)((SYMS_U8*)data.str + bin->section_array_off), bin->section_count, virt_off);
+}
+
+// imports & exports
+
+SYMS_API SYMS_U64Array
+syms_u64_array_from_null_term_u64_string(SYMS_Arena *arena, SYMS_U64 *src, SYMS_U64 *opl){
+  // scan for terminator
+  SYMS_U64 *ptr = src;
+  for (; *ptr != 0 && ptr < opl; ptr += 1);
+  
+  // make copy of the array
+  SYMS_U64Array result = {0};
+  result.count = (SYMS_U64)(ptr - src);
+  result.u64 = syms_push_array(arena, SYMS_U64, result.count);
+  syms_memmove(result.u64, src, result.count*sizeof(SYMS_U64));
+  return(result);
+}
+
+SYMS_API SYMS_U64Array
+syms_u64_array_from_null_term_u32_string(SYMS_Arena *arena, SYMS_U32 *src, SYMS_U32 *opl){
+  // scan for terminator
+  SYMS_U32 *ptr = src;
+  for (; *ptr != 0 && ptr < opl; ptr += 1);
+  
+  // make copy of the array (and transform U32 -> U64)
+  SYMS_U64Array result = {0};
+  result.count = (SYMS_U64)(ptr - src);
+  result.u64 = syms_push_array(arena, SYMS_U64, result.count);
+  SYMS_U64 *o = result.u64;
+  SYMS_U32 *i = src;
+  for (; i < ptr; o += 1, i += 1){
+    *o = *i;
+  }
+  return(result);
+}
+
+SYMS_API SYMS_U64Array
+syms_pe_u64_array_from_null_term_addr_string(SYMS_Arena *arena, SYMS_String8 data,
+                                             SYMS_PeBinAccel *bin, SYMS_U64 foff){
+  SYMS_U64Array result = {0};
+  if (bin->is_pe32){
+    SYMS_U32* src = (SYMS_U32*)(data.str + foff);
+    SYMS_U32* opl = (SYMS_U32*)(data.str + SYMS_AlignDownPow2(data.size, 4));
+    result = syms_u64_array_from_null_term_u32_string(arena, src, opl);
+  }
+  else {
+    SYMS_U64* src = (SYMS_U64*)(data.str + foff);
+    SYMS_U64* opl = (SYMS_U64*)(data.str + SYMS_AlignDownPow2(data.size, 8));
+    result = syms_u64_array_from_null_term_u64_string(arena, src, opl);
+  }
+  return(result);
+}
+
+SYMS_API SYMS_PeImportName
+syms_pe_import_name_from_name_entry(SYMS_Arena *arena, SYMS_String8 data, SYMS_PeBinAccel *bin,
+                                    SYMS_U64 name_entry){
+  SYMS_PeImportName result = {0};
+  
+  //- determine case
+  SYMS_B32 is_ordinal = syms_false;
+  if (bin->is_pe32){
+    is_ordinal = !!(name_entry & (1 << 31));
+  }
+  else{
+    is_ordinal = !!(name_entry & ((SYMS_U64)1 << 63));
+  }
+  
+  //- ordinal case
+  if (is_ordinal){
+    result.lookup = SYMS_PeImportLookup_Ordinal;
+    result.ordinal = (name_entry & SYMS_U16_MAX);
+  }
+  
+  //- voff case
+  else{
+    // TODO(allen): accelerate voff -> foff
+    SYMS_U64 foff = syms_pe_bin_foff_from_voff(data, bin, name_entry);
+    
+    // setup read helper
+    void *base = (void*)data.str;
+    SYMS_U64Range range = syms_make_u64_range(0, data.size);
+    
+    // read hint & string
+    SYMS_U16 hint = 0;
+    syms_based_range_read_struct(base, range, foff, &hint);
+    SYMS_String8 name = syms_based_range_read_string(base, range, foff + sizeof(SYMS_U16));
+    
+    // fill result
+    result.lookup = SYMS_PeImportLookup_NameHint;
+    result.name = syms_push_string_copy(arena, name);
+    result.ordinal = hint;
+  }
+  
+  return(result);
+}
+
+SYMS_API SYMS_PeImportNameArray
+syms_pe_import_name_array_from_name_entries(SYMS_Arena *arena, SYMS_String8 data, SYMS_PeBinAccel *bin,
+                                            SYMS_U64Array name_entries){
+  //- allocate output array
+  SYMS_PeImportNameArray result = {0};
+  result.names = syms_push_array_zero(arena, SYMS_PeImportName, name_entries.count);
+  result.count = name_entries.count;
+  
+  //- setup read helper
+  void *base = (void*)data.str;
+  SYMS_U64Range range = syms_make_u64_range(0, data.size);
+  
+  //- determine ordinal mask
+  SYMS_U64 ordinal_mask = ((SYMS_U64)1 << 63);
+  if (bin->is_pe32){
+    ordinal_mask = (1 << 31);
+  }
+  
+  //- fill output
+  SYMS_PeImportName *name_ptr = result.names;
+  SYMS_U64 *name_entry_ptr = name_entries.u64;
+  for (SYMS_U64 i = 0;
+       i < name_entries.count;
+       i += 1, name_ptr += 1, name_entry_ptr += 1){
+    SYMS_U64 name_entry = *name_entry_ptr;
+    SYMS_B32 is_ordinal = !!(name_entry & ordinal_mask);
+    
+    if (is_ordinal){
+      name_ptr->lookup = SYMS_PeImportLookup_Ordinal;
+      name_ptr->ordinal = (name_entry & SYMS_U16_MAX);
+    }
+    else{
+      // TODO(allen): accelerate voff -> foff
+      SYMS_U64 foff = syms_pe_bin_foff_from_voff(data, bin, name_entry);
+      
+      // read hint & string
+      SYMS_U16 hint = 0;
+      syms_based_range_read_struct(base, range, foff, &hint);
+      SYMS_String8 name = syms_based_range_read_string(base, range, foff + sizeof(SYMS_U16));
+      
+      // fill result
+      name_ptr->lookup = SYMS_PeImportLookup_NameHint;
+      name_ptr->name = syms_push_string_copy(arena, name);
+      name_ptr->ordinal = hint;
+    }
+  }
+  
+  return(result);
+}
+
+SYMS_API SYMS_PeImportDllArray
+syms_pe_regular_imports_from_bin(SYMS_Arena *arena, SYMS_String8 data, SYMS_PeBinAccel *bin){
+  //- grab ranges
+  void *base = (void *)data.str;
+  SYMS_U64Range dir_range = bin->data_dir_franges[SYMS_PeDataDirectoryIndex_IMPORT];
+  SYMS_U64Range data_range = syms_make_u64_range(0, data.size);
+  
+  //- get import entries array
+  SYMS_PeImportEntry *import_entries = (SYMS_PeImportEntry*)(data.str + dir_range.min);
+  
+  //- count imports
+  SYMS_U64 max_import_count = (dir_range.max - dir_range.min)/sizeof(SYMS_PeImportEntry);
+  SYMS_U64 import_count = max_import_count;
+  {
+    SYMS_PeImportEntry *ptr = import_entries;
+    SYMS_PeImportEntry *opl = import_entries + max_import_count;
+    for (; ptr < opl; ptr += 1){
+      if (syms_memisnull_struct(ptr)){
+        break;
+      }
+    }
+    import_count = (SYMS_U64)(ptr - import_entries);
+  }
+  
+  //- allocate output array
+  SYMS_PeImportDllArray result = {0};
+  result.import_dlls = syms_push_array_zero(arena, SYMS_PeImportDll, import_count);
+  result.count = import_count;
+  
+  //- fill output
+  SYMS_PeImportDll *dll = result.import_dlls;
+  SYMS_PeImportEntry *import_entry = import_entries;
+  SYMS_PeImportEntry *opl = import_entries + import_count;
+  for (; import_entry < opl; import_entry += 1, dll += 1){
+    // TODO(allen): accelerate voff -> foff
+    
+    // get name
+    SYMS_String8 name = {0};
+    {
+      SYMS_U64 foff = syms_pe_bin_foff_from_voff(data, bin, import_entry->name_voff);
+      SYMS_String8 name_raw = syms_based_range_read_string(base, data_range, foff);
+      name = syms_push_string_copy(arena, name_raw);
+    }
+    
+    // get name table
+    SYMS_PeImportNameArray name_table = {0};
+    {
+      SYMS_U64 foff = syms_pe_bin_foff_from_voff(data, bin, import_entry->lookup_table_voff);
+      SYMS_U64Array entries = syms_pe_u64_array_from_null_term_addr_string(arena, data, bin, foff);
+      name_table = syms_pe_import_name_array_from_name_entries(arena, data, bin, entries);
+    }
+    
+    // fill dll slot
+    dll->name = name;
+    dll->name_table = name_table;
+  }
+  
+  return(result);
+}
+
+SYMS_API SYMS_PeImportDllArray
+syms_pe_delayed_imports_from_bin(SYMS_Arena *arena, SYMS_String8 data, SYMS_PeBinAccel *bin){
+  SYMS_ArenaTemp scratch = syms_get_scratch(&arena, 1);
+  
+  //- grab prerequisites
+  void         *base       = (void *)data.str;
+  SYMS_U64Range data_range = syms_make_u64_range(0, data.size);
+  SYMS_U64Range dir_range  = bin->data_dir_franges[SYMS_PeDataDirectoryIndex_DELAY_IMPORT];
+  
+  //- get import entries array
+  SYMS_PeDelayedImportEntry *import_entries = (SYMS_PeDelayedImportEntry*)(data.str + dir_range.min);
+  
+  //- count imports
+  SYMS_U64 max_import_count = (dir_range.max - dir_range.min)/sizeof(SYMS_PeDelayedImportEntry);
+  SYMS_U64 import_count = max_import_count;
+  {
+    SYMS_PeDelayedImportEntry *ptr = import_entries;
+    SYMS_PeDelayedImportEntry *opl = import_entries + max_import_count;
+    for (; ptr < opl; ptr += 1){
+      if (syms_memisnull_struct(ptr)){
+        break;
+      }
+    }
+    import_count = (SYMS_U64)(ptr - import_entries);
+  }
+  
+  //- allocate output array
+  SYMS_PeImportDllArray result = {0};
+  result.import_dlls = syms_push_array_zero(arena, SYMS_PeImportDll, import_count);
+  result.count = import_count;
+  
+  //- parse directory
+  SYMS_PeImportDll *dll = result.import_dlls;
+  SYMS_PeDelayedImportEntry *import_ptr = import_entries;
+  SYMS_PeDelayedImportEntry *opl = import_entries + import_count;
+  for (SYMS_U64 i = 0; i < import_count; i += 1, import_ptr += 1, dll += 1){
+    // TODO(allen): accelerate voff -> foff
+    
+    //- get name
+    SYMS_String8 name = {0};
+    {
+      SYMS_U64 foff = syms_pe_bin_foff_from_voff(data, bin, import_ptr->name_voff);
+      SYMS_String8 raw_name = syms_based_range_read_string(base, data_range, foff);
+      name = syms_push_string_copy(arena, raw_name);
+    }
+    
+    //- get name table
+    SYMS_PeImportNameArray name_table = {0};
+    {
+      SYMS_U64 foff = syms_pe_bin_foff_from_voff(data, bin, import_ptr->name_table_voff);
+      SYMS_U64Array entries = syms_pe_u64_array_from_null_term_addr_string(scratch.arena, data, bin, foff);
+      name_table = syms_pe_import_name_array_from_name_entries(arena, data, bin, entries);
+    }
+    
+    //- bound imports
+    SYMS_U64Array bound_table = {0};
+    {
+      SYMS_U64 foff = syms_pe_bin_foff_from_voff(data, bin, import_ptr->bound_table_voff);
+      bound_table = syms_pe_u64_array_from_null_term_addr_string(arena, data, bin, foff);
+    }
+    
+    //- build node
+    dll->name         = name;
+    dll->name_table   = name_table;
+    dll->bound_table  = bound_table;
+  }
+  
+  syms_release_scratch(scratch);
+  
+  return(result);
+}
+
+SYMS_API SYMS_String8
+syms_pe_binary_name_from_bin(SYMS_Arena *arena, SYMS_String8 data, SYMS_PeBinAccel *bin){
+  // TODO(allen): accelerate voff -> foff
+  SYMS_U64Range dir_range  = bin->data_dir_franges[SYMS_PeDataDirectoryIndex_EXPORT];
+  SYMS_U64Range data_range = syms_make_u64_range(0, data.size);
+  void         *base       = (void*)data.str;
+  SYMS_PeExportTable *header = (SYMS_PeExportTable*)(data.str + dir_range.min);
+  SYMS_U32 binary_name_foff = syms_pe_bin_foff_from_voff(data, bin, header->name_voff);
+  SYMS_String8 binary_name = syms_based_range_read_string(base, data_range, binary_name_foff);
+  SYMS_String8 result = syms_push_string_copy(arena, binary_name);
+  return(result);
 }
 
 #endif //SYMS_PE_PARSER_C
+
diff --git a/Engine/Source/Runtime/SymsLib/syms/core/pe/syms_pe_parser.h b/Engine/Source/Runtime/SymsLib/syms/core/pe/syms_pe_parser.h
--- a/Engine/Source/Runtime/SymsLib/syms/core/pe/syms_pe_parser.h
+++ b/Engine/Source/Runtime/SymsLib/syms/core/pe/syms_pe_parser.h
@@ -4,121 +4,121 @@
 #ifndef SYMS_PE_PARSER_H
 #define SYMS_PE_PARSER_H
 
-// TODO(allen): rename? split PE and COFF accelerators?
-
 ////////////////////////////////
-// NOTE(allen): PE Parser Types
+//~ allen: PE Parser Types
 
 typedef struct SYMS_PeFileAccel{
   SYMS_FileFormat format;
   SYMS_U32 coff_off;
 } SYMS_PeFileAccel;
 
-typedef struct SYMS_CoffFileAccel{
-  SYMS_FileFormat format;
-  // TODO(allen): 
-} SYMS_CoffFileAccel;
-
 typedef struct SYMS_PeBinAccel{
   SYMS_FileFormat format;
   SYMS_U64 image_base;
+  SYMS_U64 entry_point;
+  SYMS_B32 is_pe32;
+  SYMS_U64 virt_section_align;
+  SYMS_U64 file_section_align;
   SYMS_U64 section_array_off;
   SYMS_U64 section_count;
+  SYMS_U64 symbol_array_off;
+  SYMS_U64 symbol_count;
+  SYMS_U64 string_table_off;
   SYMS_U64 dbg_path_off;
   SYMS_U64 dbg_path_size;
   SYMS_PeGuid dbg_guid;
   SYMS_U32 dbg_age;
+  SYMS_U32 dbg_time;
   SYMS_Arch arch;
+  SYMS_U64Range *data_dir_franges;
   SYMS_U32 data_dir_count;
-  SYMS_U64Range *data_dirs_virt;
-  SYMS_U64Range *data_dirs_file;
 } SYMS_PeBinAccel;
 
-typedef struct SYMS_CoffBinAccel{
-  SYMS_FileFormat format;
-  SYMS_U64 image_base;
-  SYMS_U64 section_array_off;
-  SYMS_U64 section_count;
-  SYMS_Arch arch;
-} SYMS_CoffBinAccel;
-
-enum{
-  SYMS_PeLookupBy_Ordinal,
-  SYMS_PeLookupBy_NameHint
-};
-typedef SYMS_U32 SYMS_PeImportLookupBy;
-
-typedef struct SYMS_PeImportNode
-{
-  struct SYMS_PeImportNode *next;
-  SYMS_PeImportLookupBy lookup_by;
-  SYMS_U16 ordinal;
+////////////////////////////////
+//~ allen: Imports
+
+typedef enum SYMS_PeImportLookup{
+  SYMS_PeImportLookup_NULL,
+  SYMS_PeImportLookup_Ordinal,
+  SYMS_PeImportLookup_NameHint,
+  SYMS_PeImportLookup_COUNT,
+} SYMS_PeImportLookup;
+
+typedef struct SYMS_PeImportName{
   SYMS_String8 name;
-  SYMS_U16 hint;
-} SYMS_PeImportNode;
+  SYMS_PeImportLookup lookup;
+  SYMS_U16 ordinal;
+  SYMS_U16 _padding_;
+} SYMS_PeImportName;
+
+typedef struct SYMS_PeImportNameArray{
+  SYMS_PeImportName *names;
+  SYMS_U64 count;
+} SYMS_PeImportNameArray;
 
-typedef struct SYMS_PeImportDllNode
-{
-  struct SYMS_PeImportDllNode *next;
+typedef struct SYMS_PeImportDll{
   SYMS_String8 name;
-  SYMS_PeImportDirectoryEntry import_entry;
-  SYMS_PeImportNode *first_import;
-  SYMS_PeImportNode *last_import;
-  SYMS_U32 import_count;
-} SYMS_PeImportDllNode;
+  SYMS_PeImportNameArray name_table;
+  SYMS_U64Array bound_table;
+} SYMS_PeImportDll;
+
+typedef struct SYMS_PeImportDllArray{
+  SYMS_PeImportDll *import_dlls;
+  SYMS_U64 count;
+} SYMS_PeImportDllArray;
+
 
 ////////////////////////////////
-// NOTE(allen): PE Parser Functions
+//~ allen: PE Parser Functions
 
 // accelerators
 SYMS_API SYMS_PeFileAccel* syms_pe_file_accel_from_data(SYMS_Arena *arena, SYMS_String8 data);
 SYMS_API SYMS_PeBinAccel*  syms_pe_bin_accel_from_file(SYMS_Arena *arena, SYMS_String8 data,
                                                        SYMS_PeFileAccel *file);
 
-SYMS_API SYMS_CoffFileAccel* syms_coff_file_accel_from_data(SYMS_Arena *arena, SYMS_String8 data);
-SYMS_API SYMS_CoffBinAccel*  syms_coff_bin_accel_from_file(SYMS_Arena *arena, SYMS_String8 data,
-                                                           SYMS_CoffFileAccel *file);
-
-
-// arch
-SYMS_API SYMS_Arch syms_pe_arch_from_bin(SYMS_PeBinAccel *bin);
-SYMS_API SYMS_Arch syms_coff_arch_from_bin(SYMS_CoffBinAccel *bin);
-
-// external info
-SYMS_API SYMS_ExtFileList syms_pe_ext_file_list_from_bin(SYMS_Arena *arena, SYMS_String8 data,
-                                                         SYMS_PeBinAccel *bin);
-
-// binary secs
-//  pe specific
-SYMS_API SYMS_CoffSection syms_pe_coff_section(SYMS_String8 data, SYMS_PeBinAccel *bin, SYMS_U64 n);
-
-SYMS_API SYMS_SecInfoArray syms_pe_coff_sec_info_array_from_data(SYMS_Arena *arena, SYMS_String8 data,
-                                                                 SYMS_U64 sec_array_off, SYMS_U64 sec_count);
-
-
-//  main api
+SYMS_API SYMS_Arch         syms_pe_arch_from_bin(SYMS_PeBinAccel *bin);
+SYMS_API SYMS_ExtFileList  syms_pe_ext_file_list_from_bin(SYMS_Arena *arena, SYMS_String8 data,
+                                                          SYMS_PeBinAccel *bin);
 SYMS_API SYMS_SecInfoArray syms_pe_sec_info_array_from_bin(SYMS_Arena *arena, SYMS_String8 data,
                                                            SYMS_PeBinAccel *bin);
-SYMS_API SYMS_SecInfoArray syms_coff_sec_info_array_from_bin(SYMS_Arena *arena, SYMS_String8 data,
-                                                             SYMS_CoffBinAccel *bin);
-
-// imports & exports
-//  main api
-SYMS_API SYMS_ImportArray syms_pe_imports_from_bin(SYMS_Arena *arena, SYMS_String8 data, SYMS_PeBinAccel *bin);
-SYMS_API SYMS_ExportArray syms_pe_exports_from_bin(SYMS_Arena *arena, SYMS_String8 data, SYMS_PeBinAccel *bin);
+SYMS_API SYMS_U64          syms_pe_entry_point_voff_from_bin(SYMS_PeBinAccel *bin);
+SYMS_API SYMS_ImportArray  syms_pe_imports_from_bin(SYMS_Arena *arena, SYMS_String8 data, SYMS_PeBinAccel *bin);
+SYMS_API SYMS_ExportArray  syms_pe_exports_from_bin(SYMS_Arena *arena, SYMS_String8 data, SYMS_PeBinAccel *bin);
 
 ////////////////////////////////
-// NOTE(allen): PE Specific Helpers
+//~ allen: PE Specific Helpers
 
+// pdata
 SYMS_API SYMS_U64 syms_pe_binary_search_intel_pdata(SYMS_String8 data, SYMS_PeBinAccel *bin, SYMS_U64 voff);
-
+// sections
 SYMS_API SYMS_U64 syms_pe_sec_number_from_voff(SYMS_String8 data, SYMS_PeBinAccel *bin, SYMS_U64 voff);
 
 SYMS_API void*    syms_pe_ptr_from_sec_number(SYMS_String8 data, SYMS_PeBinAccel *bin, SYMS_U64 n);
 SYMS_API void*    syms_pe_ptr_from_foff(SYMS_String8 data, SYMS_PeBinAccel *bin, SYMS_U64 foff);
 SYMS_API void*    syms_pe_ptr_from_voff(SYMS_String8 data, SYMS_PeBinAccel *bin, SYMS_U64 voff);
 
-SYMS_API SYMS_U64 syms_pe_virt_off_to_file_off(SYMS_CoffSection *sections, SYMS_U32 section_count, SYMS_U64 virt_off);
-SYMS_API SYMS_U64 syms_pe_bin_virt_off_to_file_off(SYMS_String8 data, SYMS_PeBinAccel *bin, SYMS_U64 virt_off);
+SYMS_API SYMS_U64 syms_pe_foff_from_voff(SYMS_CoffSectionHeader *sections, SYMS_U32 section_count, SYMS_U64 voff);
+SYMS_API SYMS_U64 syms_pe_bin_foff_from_voff(SYMS_String8 data, SYMS_PeBinAccel *bin, SYMS_U64 voff);
+
+// imports & exports
+SYMS_API SYMS_U64Array syms_u64_array_from_null_term_u64_string(SYMS_Arena *arena, SYMS_U64 *src, SYMS_U64 *opl);
+SYMS_API SYMS_U64Array syms_u64_array_from_null_term_u32_string(SYMS_Arena *arena, SYMS_U32 *src, SYMS_U32 *opl);
+SYMS_API SYMS_U64Array syms_pe_u64_array_from_null_term_addr_string(SYMS_Arena *arena, SYMS_String8 data,
+                                                                    SYMS_PeBinAccel *bin, SYMS_U64 foff);
+
+SYMS_API SYMS_PeImportName syms_pe_import_name_from_name_entry(SYMS_Arena *arena, SYMS_String8 data,
+                                                               SYMS_PeBinAccel *bin, SYMS_U64 name_entry);
+SYMS_API SYMS_PeImportNameArray syms_pe_import_name_array_from_name_entries(SYMS_Arena *arena,
+                                                                            SYMS_String8 data,
+                                                                            SYMS_PeBinAccel *bin,
+                                                                            SYMS_U64Array name_entries);
+
+SYMS_API SYMS_PeImportDllArray syms_pe_regular_imports_from_bin(SYMS_Arena *arena, SYMS_String8 data,
+                                                                SYMS_PeBinAccel *bin);
+SYMS_API SYMS_PeImportDllArray syms_pe_delayed_imports_from_bin(SYMS_Arena *arena, SYMS_String8 data,
+                                                                SYMS_PeBinAccel *bin);
+
+SYMS_API SYMS_String8     syms_pe_binary_name_from_bin(SYMS_Arena *arena, SYMS_String8 data,
+                                                       SYMS_PeBinAccel *bin);
 
 #endif //SYMS_PE_PARSER_H
diff --git a/Engine/Source/Runtime/SymsLib/syms/core/syms_debug_info.c b/Engine/Source/Runtime/SymsLib/syms/core/syms_debug_info.c
--- a/Engine/Source/Runtime/SymsLib/syms/core/syms_debug_info.c
+++ b/Engine/Source/Runtime/SymsLib/syms/core/syms_debug_info.c
@@ -93,6 +93,18 @@ syms_type_kind_is_user_defined(SYMS_TypeKind kind){
   return(result);
 }
 
+SYMS_API SYMS_B32
+syms_type_kind_is_record(SYMS_TypeKind kind){
+  SYMS_B32 result = (SYMS_TypeKind_Struct <= kind && kind <= SYMS_TypeKind_Union);
+  return(result);
+}
+
+SYMS_API SYMS_B32
+syms_type_kind_is_enum(SYMS_TypeKind kind){
+  SYMS_B32 result = (kind == SYMS_TypeKind_Enum);
+  return(result);
+}
+
 SYMS_API SYMS_B32
 syms_type_kind_is_forward(SYMS_TypeKind kind){
   SYMS_B32 result = (SYMS_TypeKind_ForwardStruct <= kind && kind <= SYMS_TypeKind_ForwardEnum);
diff --git a/Engine/Source/Runtime/SymsLib/syms/core/syms_debug_info.h b/Engine/Source/Runtime/SymsLib/syms/core/syms_debug_info.h
--- a/Engine/Source/Runtime/SymsLib/syms/core/syms_debug_info.h
+++ b/Engine/Source/Runtime/SymsLib/syms/core/syms_debug_info.h
@@ -10,7 +10,7 @@
 #include "syms/core/generated/syms_meta_debug_info.h"
 
 ////////////////////////////////
-//~ NOTE(rjf): Binary Info Types
+//~ rjf: Binary Info Types
 
 typedef struct SYMS_SecInfo{
   SYMS_String8 name;
@@ -338,18 +338,18 @@ typedef struct SYMS_SigInfo{
   SYMS_SymbolID this_type_id;
 } SYMS_SigInfo;
 
-typedef struct SYMS_EnumInfo{
+typedef struct SYMS_EnumMember{
   SYMS_String8 name;
   SYMS_U64 val;
-} SYMS_EnumInfo;
+} SYMS_EnumMember;
 
-typedef struct SYMS_EnumInfoArray{
-  SYMS_EnumInfo *enum_info;
+typedef struct SYMS_EnumMemberArray{
+  SYMS_EnumMember *enum_members;
   SYMS_U64 count;
-} SYMS_EnumInfoArray;
+} SYMS_EnumMemberArray;
 
 ////////////////////////////////
-// allen: Location Types
+//~ allen: Location Types
 
 typedef SYMS_U64 SYMS_LocID;
 
@@ -380,7 +380,7 @@ typedef enum SYMS_ProcLoc{
 } SYMS_ProcLoc;
 
 ////////////////////////////////
-// allen: Link Name Types
+//~ allen: Link Name Types
 
 typedef struct SYMS_LinkNameRec{
   SYMS_String8 name;
@@ -393,7 +393,7 @@ typedef struct SYMS_LinkNameRecArray{
 } SYMS_LinkNameRecArray;
 
 ////////////////////////////////
-// allen: Nil For Accelerators
+//~ allen: Nil For Accelerators
 
 SYMS_READ_ONLY SYMS_GLOBAL SYMS_FileFormat syms_format_nil = SYMS_FileFormat_Null;
 
@@ -403,7 +403,6 @@ SYMS_READ_ONLY SYMS_GLOBAL SYMS_FileFormat syms_format_nil = SYMS_FileFormat_Nul
 SYMS_C_LINKAGE_BEGIN
 
 SYMS_API SYMS_B32     syms_ext_match_key_match(SYMS_ExtMatchKey *a, SYMS_ExtMatchKey *b);
-SYMS_API SYMS_String8 syms_string_from_symbol_kind(SYMS_SymbolKind kind);
 
 SYMS_API SYMS_USID    syms_make_usid(SYMS_UnitID uid, SYMS_SymbolID sid);
 
@@ -415,6 +414,8 @@ SYMS_API SYMS_B32      syms_type_kind_is_integer(SYMS_TypeKind kind);
 SYMS_API SYMS_B32      syms_type_kind_is_signed(SYMS_TypeKind kind);
 SYMS_API SYMS_B32      syms_type_kind_is_complex(SYMS_TypeKind kind);
 SYMS_API SYMS_B32      syms_type_kind_is_user_defined(SYMS_TypeKind kind);
+SYMS_API SYMS_B32      syms_type_kind_is_record(SYMS_TypeKind kind);
+SYMS_API SYMS_B32      syms_type_kind_is_enum(SYMS_TypeKind kind);
 SYMS_API SYMS_B32      syms_type_kind_is_forward(SYMS_TypeKind kind);
 
 SYMS_API SYMS_SymbolIDArray syms_sid_array_from_list(SYMS_Arena *arena, SYMS_SymbolIDList *list);
diff --git a/Engine/Source/Runtime/SymsLib/syms/core/syms_eval.c b/Engine/Source/Runtime/SymsLib/syms/core/syms_eval.c
--- a/Engine/Source/Runtime/SymsLib/syms/core/syms_eval.c
+++ b/Engine/Source/Runtime/SymsLib/syms/core/syms_eval.c
@@ -4,7 +4,7 @@
 #define SYMS_EVAL_C
 
 ////////////////////////////////
-// NOTE(allen): Eval Meta Code
+//~ allen: Eval Meta Code
 
 #include "syms/core/generated/syms_meta_eval.c"
 
@@ -12,17 +12,17 @@
 #include "syms/core/generated/syms_meta_eval_ext.c"
 
 ////////////////////////////////
-// NOTE(allen): Eval Bytecode Helper Functions
+//~ allen: Eval Bytecode Helper Functions
 
 SYMS_API SYMS_EvalOpParams
-syms_op_params(SYMS_U64 p){
+syms_eval_op_params(SYMS_U64 p){
   SYMS_EvalOpParams result = {0};
   result.u64[0] = p;
   return(result);
 }
 
 SYMS_API SYMS_EvalOpParams
-syms_op_params_2u8(SYMS_U8 p1, SYMS_U8 p2){
+syms_eval_op_params_2u8(SYMS_U8 p1, SYMS_U8 p2){
   SYMS_EvalOpParams result = {0};
   result.u8[0] = p1;
   result.u8[1] = p2;
@@ -30,7 +30,7 @@ syms_op_params_2u8(SYMS_U8 p1, SYMS_U8 p2){
 }
 
 SYMS_API SYMS_EvalOpParams
-syms_op_params_2u16(SYMS_U16 p1, SYMS_U16 p2){
+syms_eval_op_params_2u16(SYMS_U16 p1, SYMS_U16 p2){
   SYMS_EvalOpParams result = {0};
   result.u16[0] = p1;
   result.u16[1] = p2;
@@ -38,7 +38,7 @@ syms_op_params_2u16(SYMS_U16 p1, SYMS_U16 p2){
 }
 
 SYMS_API void
-syms_op_push(SYMS_Arena *arena, SYMS_EvalOpList *list, SYMS_EvalOp op, SYMS_EvalOpParams params){
+syms_eval_op_push(SYMS_Arena *arena, SYMS_EvalOpList *list, SYMS_EvalOp op, SYMS_EvalOpParams params){
   SYMS_ASSERT(op < SYMS_EvalOp_COUNT);
   
   //- setup new op
@@ -58,7 +58,7 @@ syms_op_push(SYMS_Arena *arena, SYMS_EvalOpList *list, SYMS_EvalOp op, SYMS_Eval
 }
 
 SYMS_API void
-syms_op_push_bytecode(SYMS_Arena *arena, SYMS_EvalOpList *list, SYMS_String8 bytecode){
+syms_eval_op_push_bytecode(SYMS_Arena *arena, SYMS_EvalOpList *list, SYMS_String8 bytecode){
   //- setup new op
   SYMS_EvalOpNode *loose_op = syms_push_array(arena, SYMS_EvalOpNode, 1);
   loose_op->op = SYMS_EvalIRExtKind_Bytecode;
@@ -70,23 +70,23 @@ syms_op_push_bytecode(SYMS_Arena *arena, SYMS_EvalOpList *list, SYMS_String8 byt
 }
 
 SYMS_API void
-syms_op_list_concat_in_place(SYMS_EvalOpList *left, SYMS_EvalOpList *right){
-  left->byte_count += right->byte_count;
-  if (left->last != 0){
-    if (right->last != 0){
-      left->last->next = right->first;
-      left->last = right->last;
+syms_eval_op_list_concat_in_place(SYMS_EvalOpList *left_dst, SYMS_EvalOpList *right_destroyed){
+  left_dst->byte_count += right_destroyed->byte_count;
+  if (left_dst->last != 0){
+    if (right_destroyed->last != 0){
+      left_dst->last->next = right_destroyed->first;
+      left_dst->last = right_destroyed->last;
     }
   }
   else{
-    left->first = right->first;
-    left->last = right->last;
+    left_dst->first = right_destroyed->first;
+    left_dst->last = right_destroyed->last;
   }
-  syms_memzero_struct(right);
+  syms_memzero_struct(right_destroyed);
 }
 
 SYMS_API void
-syms_op_encode_u(SYMS_Arena *arena, SYMS_EvalOpList *list, SYMS_U64 u){
+syms_eval_op_encode_u(SYMS_Arena *arena, SYMS_EvalOpList *list, SYMS_U64 u){
   SYMS_EvalOp op = SYMS_EvalOp_ConstU64;
   if (u <= SYMS_U8_MAX){
     op = SYMS_EvalOp_ConstU8;
@@ -97,11 +97,11 @@ syms_op_encode_u(SYMS_Arena *arena, SYMS_EvalOpList *list, SYMS_U64 u){
   else if (u <= SYMS_U32_MAX){
     op = SYMS_EvalOp_ConstU32;
   }
-  syms_op_push(arena, list, op, syms_op_params(u));
+  syms_eval_op_push(arena, list, op, syms_eval_op_params(u));
 }
 
 SYMS_API void
-syms_op_encode_s(SYMS_Arena *arena, SYMS_EvalOpList *list, SYMS_S64 s){
+syms_eval_op_encode_s(SYMS_Arena *arena, SYMS_EvalOpList *list, SYMS_S64 s){
   SYMS_U64 size = 64;
   SYMS_EvalOp op = SYMS_EvalOp_ConstU64;
   if (SYMS_S8_MIN <= s && s <= SYMS_S8_MAX){
@@ -116,26 +116,55 @@ syms_op_encode_s(SYMS_Arena *arena, SYMS_EvalOpList *list, SYMS_S64 s){
     op = SYMS_EvalOp_ConstU32;
     size = 32;
   }
-  syms_op_push(arena, list, op, syms_op_params((SYMS_U64)s));
+  syms_eval_op_push(arena, list, op, syms_eval_op_params((SYMS_U64)s));
   if (size < 64){
-    syms_op_push(arena, list, SYMS_EvalOp_TruncSigned, syms_op_params(size));
+    syms_eval_op_push(arena, list, SYMS_EvalOp_TruncSigned, syms_eval_op_params(size));
   }
 }
 
 SYMS_API void
-syms_op_encode_reg_section(SYMS_Arena *arena, SYMS_EvalOpList *list, SYMS_RegSection sec){
+syms_eval_op_encode_reg_section(SYMS_Arena *arena, SYMS_EvalOpList *list, SYMS_RegSection sec){
   SYMS_U64 u64 = sec.off | sec.size << 16;
-  syms_op_encode_u(arena, list, u64);
+  syms_eval_op_encode_u(arena, list, u64);
 }
 
 SYMS_API void
-syms_op_encode_reg(SYMS_Arena *arena, SYMS_EvalOpList *list, SYMS_Arch arch, SYMS_RegID reg_id){
+syms_eval_op_encode_reg(SYMS_Arena *arena, SYMS_EvalOpList *list, SYMS_Arch arch, SYMS_RegID reg_id){
   SYMS_RegSection sec = syms_reg_section_from_reg_id(arch, reg_id);
-  syms_op_encode_reg_section(arena, list, sec);
+  syms_eval_op_encode_reg_section(arena, list, sec);
 }
 
+//- bytecode analysis functions
+
+SYMS_API SYMS_B32
+syms_eval_opcode_type_compatible(SYMS_EvalOp op, SYMS_EvalTypeGroup group){
+  SYMS_B32 result = syms_false;
+  switch (op){
+    case SYMS_EvalOp_Neg:case SYMS_EvalOp_Add:case SYMS_EvalOp_Sub:
+    case SYMS_EvalOp_Mul:case SYMS_EvalOp_Div:
+    case SYMS_EvalOp_EqEq:case SYMS_EvalOp_NtEq:
+    case SYMS_EvalOp_LsEq:case SYMS_EvalOp_GrEq:case SYMS_EvalOp_Less:case SYMS_EvalOp_Grtr:
+    {
+      if (group != SYMS_EvalTypeGroup_Other){
+        result = syms_true;
+      }
+    }break;
+    case SYMS_EvalOp_Mod:case SYMS_EvalOp_LShift:case SYMS_EvalOp_RShift:
+    case SYMS_EvalOp_BitNot:case SYMS_EvalOp_BitAnd:case SYMS_EvalOp_BitXor:case SYMS_EvalOp_BitOr:
+    case SYMS_EvalOp_LogNot:case SYMS_EvalOp_LogAnd:case SYMS_EvalOp_LogOr: 
+    {
+      if (group == SYMS_EvalTypeGroup_S || group == SYMS_EvalTypeGroup_U){
+        result = syms_true;
+      }
+    }break;
+  }
+  return(result);
+}
+
+//- bytecode encoder
+
 SYMS_API SYMS_String8
-syms_bytecode_from_op_list(SYMS_Arena *arena, SYMS_EvalOpList *list){
+syms_eval_bytecode_from_op_list(SYMS_Arena *arena, SYMS_EvalOpList *list){
   // allocate output
   SYMS_String8 result = {0};
   result.size = list->byte_count;
@@ -186,35 +215,10 @@ syms_bytecode_from_op_list(SYMS_Arena *arena, SYMS_EvalOpList *list){
   return(result);
 }
 
-SYMS_API SYMS_B32
-syms_eval_opcode_type_compatible(SYMS_EvalOp op, SYMS_EvalTypeGroup group){
-  SYMS_B32 result = syms_false;
-  switch (op){
-    case SYMS_EvalOp_Neg:case SYMS_EvalOp_Add:case SYMS_EvalOp_Sub:
-    case SYMS_EvalOp_Mul:case SYMS_EvalOp_Div:
-    case SYMS_EvalOp_EqEq:case SYMS_EvalOp_NtEq:
-    case SYMS_EvalOp_LsEq:case SYMS_EvalOp_GrEq:case SYMS_EvalOp_Less:case SYMS_EvalOp_Grtr:
-    {
-      if (group != SYMS_EvalTypeGroup_Other){
-        result = syms_true;
-      }
-    }break;
-    case SYMS_EvalOp_Mod:case SYMS_EvalOp_LShift:case SYMS_EvalOp_RShift:
-    case SYMS_EvalOp_BitNot:case SYMS_EvalOp_BitAnd:case SYMS_EvalOp_BitXor:case SYMS_EvalOp_BitOr:
-    case SYMS_EvalOp_LogNot:case SYMS_EvalOp_LogAnd:case SYMS_EvalOp_LogOr: 
-    {
-      if (group == SYMS_EvalTypeGroup_S || group == SYMS_EvalTypeGroup_U){
-        result = syms_true;
-      }
-    }break;
-  }
-  return(result);
-}
-
 //- ir tree helpers functions
 
 SYMS_API SYMS_EvalIRTree*
-syms_ir_tree_const_u(SYMS_Arena *arena, SYMS_U64 v){
+syms_eval_ir_tree_const_u(SYMS_Arena *arena, SYMS_U64 v){
   // choose encoding op
   SYMS_EvalOp op = SYMS_EvalOp_ConstU64;
   if (v < 0x100){
@@ -230,38 +234,38 @@ syms_ir_tree_const_u(SYMS_Arena *arena, SYMS_U64 v){
   // make the tree node
   SYMS_EvalIRTree *result = syms_push_array_zero(arena, SYMS_EvalIRTree, 1);
   result->op = op;
-  result->params = syms_op_params(v);
+  result->params = syms_eval_op_params(v);
   return(result);
 }
 
 SYMS_API SYMS_EvalIRTree*
-syms_ir_tree_unary_op(SYMS_Arena *arena, SYMS_EvalOp op, SYMS_EvalTypeGroup group, SYMS_EvalIRTree *c){
+syms_eval_ir_tree_unary_op(SYMS_Arena *arena, SYMS_EvalOp op, SYMS_EvalTypeGroup group, SYMS_EvalIRTree *c){
   SYMS_EvalIRTree *result = syms_push_array_zero(arena, SYMS_EvalIRTree, 1);
   result->op = op;
-  result->params = syms_op_params(group);
+  result->params = syms_eval_op_params(group);
   result->children[0] = c;
   return(result);
 }
 
 SYMS_API SYMS_EvalIRTree*
-syms_ir_tree_binary_op(SYMS_Arena *arena, SYMS_EvalOp op, SYMS_EvalTypeGroup group,
-                       SYMS_EvalIRTree *l, SYMS_EvalIRTree *r){
+syms_eval_ir_tree_binary_op(SYMS_Arena *arena, SYMS_EvalOp op, SYMS_EvalTypeGroup group,
+                            SYMS_EvalIRTree *l, SYMS_EvalIRTree *r){
   SYMS_EvalIRTree *result = syms_push_array_zero(arena, SYMS_EvalIRTree, 1);
   result->op = op;
-  result->params = syms_op_params(group);
+  result->params = syms_eval_op_params(group);
   result->children[0] = l;
   result->children[1] = r;
   return(result);
 }
 
 SYMS_API SYMS_EvalIRTree*
-syms_ir_tree_binary_op_u(SYMS_Arena *arena, SYMS_EvalOp op, SYMS_EvalIRTree *l, SYMS_EvalIRTree *r){
-  SYMS_EvalIRTree *result = syms_ir_tree_binary_op(arena, op, SYMS_EvalTypeGroup_U, l, r);
+syms_eval_ir_tree_binary_op_u(SYMS_Arena *arena, SYMS_EvalOp op, SYMS_EvalIRTree *l, SYMS_EvalIRTree *r){
+  SYMS_EvalIRTree *result = syms_eval_ir_tree_binary_op(arena, op, SYMS_EvalTypeGroup_U, l, r);
   return(result);
 }
 
 SYMS_API SYMS_EvalIRTree*
-syms_ir_tree_conditional(SYMS_Arena *arena, SYMS_EvalIRTree *c, SYMS_EvalIRTree *l, SYMS_EvalIRTree *r){
+syms_eval_ir_tree_conditional(SYMS_Arena *arena, SYMS_EvalIRTree *c, SYMS_EvalIRTree *l, SYMS_EvalIRTree *r){
   SYMS_EvalIRTree *result = syms_push_array_zero(arena, SYMS_EvalIRTree, 1);
   result->op = SYMS_EvalOp_Cond;
   result->children[0] = c;
@@ -271,7 +275,7 @@ syms_ir_tree_conditional(SYMS_Arena *arena, SYMS_EvalIRTree *c, SYMS_EvalIRTree
 }
 
 SYMS_API SYMS_EvalIRTree*
-syms_ir_tree_bytecode_no_copy(SYMS_Arena *arena, SYMS_String8 bytecode){
+syms_eval_ir_tree_bytecode_no_copy(SYMS_Arena *arena, SYMS_String8 bytecode){
   SYMS_EvalIRTree *result = syms_push_array_zero(arena, SYMS_EvalIRTree, 1);
   result->op = SYMS_EvalIRExtKind_Bytecode;
   result->params.data = bytecode;
