Upgrade SymsLib to latest version

diff --git a/Engine/Source/Runtime/SymsLib/syms/core/syms_eval.h b/Engine/Source/Runtime/SymsLib/syms/core/syms_eval.h
--- a/Engine/Source/Runtime/SymsLib/syms/core/syms_eval.h
+++ b/Engine/Source/Runtime/SymsLib/syms/core/syms_eval.h
@@ -5,12 +5,12 @@
 #define SYMS_EVAL_H
 
 ////////////////////////////////
-// NOTE(allen): Eval Meta Code
+//~ allen: Syms Eval Meta Code
 
 #include "syms/core/generated/syms_meta_eval.h"
 
 ////////////////////////////////
-// NOTE(allen): Eval Bytecode Helper Types
+//~ allen: Syms Eval Bytecode Helper Types
 
 typedef enum SYMS_EvalMode{
   SYMS_EvalMode_Nil,
@@ -62,51 +62,55 @@ typedef struct SYMS_EvalOpList{
 typedef struct SYMS_Location{
   SYMS_EvalOpList op_list;
   SYMS_EvalMode mode;
+  SYMS_B32 is_parameter;
 } SYMS_Location;
 
 ////////////////////////////////
-// NOTE(allen): 
+//~ allen: Syms Eval Nil Values
 
 SYMS_READ_ONLY SYMS_GLOBAL SYMS_EvalIRTree syms_eval_ir_tree_nil = {0};
 
 ////////////////////////////////
-// NOTE(allen): Eval Bytecode Helper Functions
+//~ allen: Syms Eval Functions
 
-// TODO(allen): naming pass "syms_eval_*"
+//- bytecode helper functions
 
-SYMS_API SYMS_EvalOpParams syms_op_params(SYMS_U64 p);
-SYMS_API SYMS_EvalOpParams syms_op_params_2u8(SYMS_U8 p1, SYMS_U8 p2);
-SYMS_API SYMS_EvalOpParams syms_op_params_2u16(SYMS_U16 p1, SYMS_U16 p2);
-SYMS_API SYMS_EvalOpParams syms_op_params_bytecode(SYMS_String8 bytecode);
+SYMS_API SYMS_EvalOpParams syms_eval_op_params(SYMS_U64 p);
+SYMS_API SYMS_EvalOpParams syms_eval_op_params_2u8(SYMS_U8 p1, SYMS_U8 p2);
+SYMS_API SYMS_EvalOpParams syms_eval_op_params_2u16(SYMS_U16 p1, SYMS_U16 p2);
 
-SYMS_API void syms_op_push(SYMS_Arena *arena, SYMS_EvalOpList *list, SYMS_EvalOp op,
-                           SYMS_EvalOpParams params);
-SYMS_API void syms_op_push_bytecode(SYMS_Arena *arena, SYMS_EvalOpList *list, SYMS_String8 bytecode);
-SYMS_API void syms_op_list_concat_in_place(SYMS_EvalOpList *left, SYMS_EvalOpList *right);
+SYMS_API void syms_eval_op_push(SYMS_Arena *arena, SYMS_EvalOpList *list, SYMS_EvalOp op,
+                                SYMS_EvalOpParams params);
+SYMS_API void syms_eval_op_push_bytecode(SYMS_Arena *arena, SYMS_EvalOpList *list, SYMS_String8 bytecode);
+SYMS_API void syms_eval_op_list_concat_in_place(SYMS_EvalOpList *left_dst, SYMS_EvalOpList *right_destroyed);
 
-SYMS_API void syms_op_encode_u(SYMS_Arena *arena, SYMS_EvalOpList *list, SYMS_U64 u);
-SYMS_API void syms_op_encode_s(SYMS_Arena *arena, SYMS_EvalOpList *list, SYMS_S64 s);
+SYMS_API void syms_eval_op_encode_u(SYMS_Arena *arena, SYMS_EvalOpList *list, SYMS_U64 u);
+SYMS_API void syms_eval_op_encode_s(SYMS_Arena *arena, SYMS_EvalOpList *list, SYMS_S64 s);
 
-SYMS_API void syms_op_encode_reg_section(SYMS_Arena *arena, SYMS_EvalOpList *list, SYMS_RegSection sec);
-SYMS_API void syms_op_encode_reg(SYMS_Arena *arena, SYMS_EvalOpList *list, SYMS_Arch arch,
-                                 SYMS_RegID reg_id);
+SYMS_API void syms_eval_op_encode_reg_section(SYMS_Arena *arena, SYMS_EvalOpList *list, SYMS_RegSection sec);
+SYMS_API void syms_eval_op_encode_reg(SYMS_Arena *arena, SYMS_EvalOpList *list, SYMS_Arch arch,
+                                      SYMS_RegID reg_id);
 
-SYMS_API SYMS_String8 syms_bytecode_from_op_list(SYMS_Arena *arena, SYMS_EvalOpList *list);
+//- bytecode analysis functions
 
 SYMS_API SYMS_B32 syms_eval_opcode_type_compatible(SYMS_EvalOp op, SYMS_EvalTypeGroup group);
 
+//- bytecode encoder
+
+SYMS_API SYMS_String8 syms_eval_bytecode_from_op_list(SYMS_Arena *arena, SYMS_EvalOpList *list);
+
 //- ir tree helpers functions
 
-SYMS_API SYMS_EvalIRTree* syms_ir_tree_const_u(SYMS_Arena *arena, SYMS_U64 v);
-SYMS_API SYMS_EvalIRTree* syms_ir_tree_unary_op(SYMS_Arena *arena, SYMS_EvalOp op,
-                                                SYMS_EvalTypeGroup group, SYMS_EvalIRTree *c);
-SYMS_API SYMS_EvalIRTree* syms_ir_tree_binary_op(SYMS_Arena *arena, SYMS_EvalOp op,
-                                                 SYMS_EvalTypeGroup group,
-                                                 SYMS_EvalIRTree *l, SYMS_EvalIRTree *r);
-SYMS_API SYMS_EvalIRTree* syms_ir_tree_binary_op_u(SYMS_Arena *arena, SYMS_EvalOp op, SYMS_EvalIRTree *l,
-                                                   SYMS_EvalIRTree *r);
-SYMS_API SYMS_EvalIRTree* syms_ir_tree_conditional(SYMS_Arena *arena, SYMS_EvalIRTree *c,
-                                                   SYMS_EvalIRTree *l, SYMS_EvalIRTree *r);
-SYMS_API SYMS_EvalIRTree* syms_ir_tree_bytecode_no_copy(SYMS_Arena *arena, SYMS_String8 bytecode);
+SYMS_API SYMS_EvalIRTree* syms_eval_ir_tree_const_u(SYMS_Arena *arena, SYMS_U64 v);
+SYMS_API SYMS_EvalIRTree* syms_eval_ir_tree_unary_op(SYMS_Arena *arena, SYMS_EvalOp op,
+                                                     SYMS_EvalTypeGroup group, SYMS_EvalIRTree *c);
+SYMS_API SYMS_EvalIRTree* syms_eval_ir_tree_binary_op(SYMS_Arena *arena, SYMS_EvalOp op,
+                                                      SYMS_EvalTypeGroup group,
+                                                      SYMS_EvalIRTree *l, SYMS_EvalIRTree *r);
+SYMS_API SYMS_EvalIRTree* syms_eval_ir_tree_binary_op_u(SYMS_Arena *arena, SYMS_EvalOp op, SYMS_EvalIRTree *l,
+                                                        SYMS_EvalIRTree *r);
+SYMS_API SYMS_EvalIRTree* syms_eval_ir_tree_conditional(SYMS_Arena *arena, SYMS_EvalIRTree *c,
+                                                        SYMS_EvalIRTree *l, SYMS_EvalIRTree *r);
+SYMS_API SYMS_EvalIRTree* syms_eval_ir_tree_bytecode_no_copy(SYMS_Arena *arena, SYMS_String8 bytecode);
 
 #endif //SYMS_EVAL_H
diff --git a/Engine/Source/Runtime/SymsLib/syms/core/syms_parser.c b/Engine/Source/Runtime/SymsLib/syms/core/syms_parser.c
--- a/Engine/Source/Runtime/SymsLib/syms/core/syms_parser.c
+++ b/Engine/Source/Runtime/SymsLib/syms/core/syms_parser.c
@@ -24,11 +24,6 @@ syms_file_accel_from_data(SYMS_Arena *arena, SYMS_String8 data){
     result = (SYMS_FileAccel*)syms_mach_file_accel_from_data(arena, data);
   }
   
-  // harder to recognize
-  if (result->format == SYMS_FileFormat_Null){
-    result = (SYMS_FileAccel*)syms_coff_file_accel_from_data(arena, data);
-  }
-  
   SYMS_ASSERT_PARANOID(syms_parser_api_invariants());
   return(result);
 }
@@ -48,8 +43,6 @@ syms_file_is_bin(SYMS_FileAccel *file){
   switch (file->format){
     case SYMS_FileFormat_PE:
     case SYMS_FileFormat_ELF:
-    case SYMS_FileFormat_COFF16:
-    case SYMS_FileFormat_COFF32:
     {
       result = syms_true;
     }break;
@@ -73,12 +66,6 @@ syms_bin_accel_from_file(SYMS_Arena *arena, SYMS_String8 data, SYMS_FileAccel *f
       result = (SYMS_BinAccel*)syms_pe_bin_accel_from_file(arena, data, &file->pe_accel);
     }break;
     
-    case SYMS_FileFormat_COFF16:
-    case SYMS_FileFormat_COFF32:
-    {
-      result = (SYMS_BinAccel*)syms_coff_bin_accel_from_file(arena, data, &file->coff_accel);
-    }break;
-    
     case SYMS_FileFormat_ELF:
     {
       result = (SYMS_BinAccel*)syms_elf_bin_accel_from_file(arena, data, &file->elf_accel);
@@ -105,12 +92,6 @@ syms_arch_from_bin(SYMS_BinAccel *bin){
       result = syms_pe_arch_from_bin(&bin->pe_accel);
     }break;
     
-    case SYMS_FileFormat_COFF16:
-    case SYMS_FileFormat_COFF32:
-    {
-      result = syms_coff_arch_from_bin(&bin->coff_accel);
-    }break;
-    
     case SYMS_FileFormat_ELF:
     {
       result = syms_elf_arch_from_bin(&bin->elf_accel);
@@ -158,12 +139,6 @@ syms_sec_info_array_from_bin(SYMS_Arena *arena, SYMS_String8 data, SYMS_BinAccel
       result = syms_pe_sec_info_array_from_bin(arena, data, &bin->pe_accel);
     }break;
     
-    case SYMS_FileFormat_COFF16:
-    case SYMS_FileFormat_COFF32:
-    {
-      result = syms_coff_sec_info_array_from_bin(arena, data, &bin->coff_accel);
-    }break;
-    
     case SYMS_FileFormat_ELF:
     {
       result = syms_elf_sec_info_array_from_bin(arena, data, &bin->elf_accel);
@@ -206,6 +181,32 @@ syms_default_vbase_from_bin(SYMS_BinAccel *bin){
   return(result);
 }
 
+// entry point
+SYMS_API SYMS_U64
+syms_entry_point_voff_from_bin(SYMS_BinAccel *bin){
+  SYMS_ProfBegin("syms_entry_point_voff_from_bin");
+  SYMS_U64 result = 0;
+  switch (bin->format){
+    case SYMS_FileFormat_PE:
+    {
+      result = syms_pe_entry_point_voff_from_bin(&bin->pe_accel);
+    }break;
+    
+    case SYMS_FileFormat_ELF:
+    {
+      result = syms_elf_entry_point_voff_from_bin(&bin->elf_accel);
+    }break;
+    
+    case SYMS_FileFormat_MACH:
+    {
+      // TODO(allen): ?
+    }break;
+  }
+  SYMS_ASSERT_PARANOID(syms_parser_api_invariants());
+  SYMS_ProfEnd();
+  return(result);
+}
+
 // imports & exports
 SYMS_API SYMS_ImportArray
 syms_imports_from_bin(SYMS_Arena *arena, SYMS_String8 data, SYMS_BinAccel *bin){
@@ -366,9 +367,8 @@ syms_bin_is_dbg(SYMS_BinAccel *bin){
   return(result);
 }
 
-SYMS_API SYMS_DbgAccel *
-syms_dbg_accel_from_bin(SYMS_Arena *arena, SYMS_String8 data, SYMS_BinAccel *bin)
-{
+SYMS_API SYMS_DbgAccel*
+syms_dbg_accel_from_bin(SYMS_Arena *arena, SYMS_String8 data, SYMS_BinAccel *bin){
   SYMS_ProfBegin("syms_dbg_accel_from_bin");
   SYMS_DbgAccel *result = (SYMS_DbgAccel*)&syms_format_nil;
   switch(bin->format){
@@ -416,7 +416,7 @@ syms_ext_file_list_from_dbg(SYMS_Arena *arena, SYMS_String8 data, SYMS_DbgAccel
   switch (dbg->format){
     case SYMS_FileFormat_PDB:
     {
-      // NOTE(allen): PDB has no external files, do nothing.
+      // PDB has no external files, do nothing.
     }break;
     
     case SYMS_FileFormat_DWARF:
@@ -636,20 +636,23 @@ syms_uid_collated_public_symbols_from_set(SYMS_UnitSetAccel *unit_set){
 
 // symbol parsing
 SYMS_API SYMS_UnitAccel*
-syms_unit_accel_from_uid(SYMS_Arena *arena, SYMS_String8 data, SYMS_DbgAccel *dbg, SYMS_UnitSetAccel *unit_set,
-                         SYMS_UnitID uid){
+syms_unit_accel_from_uid(SYMS_Arena *arena, SYMS_String8 data, SYMS_DbgAccel *dbg,
+                         SYMS_UnitSetAccel *unit_set, SYMS_UnitID uid){
   SYMS_ProfBegin("syms_unit_accel_from_uid");
   SYMS_UnitAccel *result = (SYMS_UnitAccel*)&syms_format_nil;
   if (unit_set->format == dbg->format){
     switch (unit_set->format){
       case SYMS_FileFormat_PDB:
       {
-        result = (SYMS_UnitAccel*)syms_pdb_unit_accel_from_id(arena, data, &dbg->pdb_accel, &unit_set->pdb_accel, uid);
+        result =
+        (SYMS_UnitAccel*)syms_pdb_unit_accel_from_uid(arena, data, &dbg->pdb_accel, &unit_set->pdb_accel,
+                                                      uid);
       }break;
       
       case SYMS_FileFormat_DWARF:
       {
-        result = (SYMS_UnitAccel*)syms_dw_unit_accel_from_uid(arena, data, &dbg->dw_accel, &unit_set->dw_accel, uid);
+        result = (SYMS_UnitAccel*)syms_dw_unit_accel_from_uid(arena, data, &dbg->dw_accel, 
+                                                              &unit_set->dw_accel, uid);
       }break;
     }
   }
@@ -665,7 +668,7 @@ syms_uid_from_unit(SYMS_UnitAccel *unit){
   switch (unit->format){
     case SYMS_FileFormat_PDB:
     {
-      result = syms_pdb_uid_from_accel(&unit->pdb_accel);
+      result = syms_cv_uid_from_accel(&unit->cv_accel);
     }break;
     
     case SYMS_FileFormat_DWARF:
@@ -685,7 +688,7 @@ syms_proc_sid_array_from_unit(SYMS_Arena *arena, SYMS_UnitAccel *unit){
   switch (unit->format){
     case SYMS_FileFormat_PDB:
     {
-      result = syms_pdb_proc_sid_array_from_unit(arena, &unit->pdb_accel);
+      result = syms_cv_proc_sid_array_from_unit(arena, &unit->cv_accel);
     }break;
     
     case SYMS_FileFormat_DWARF:
@@ -705,7 +708,7 @@ syms_var_sid_array_from_unit(SYMS_Arena *arena, SYMS_UnitAccel *unit){
   switch (unit->format){
     case SYMS_FileFormat_PDB:
     {
-      result = syms_pdb_var_sid_array_from_unit(arena, &unit->pdb_accel);
+      result = syms_cv_var_sid_array_from_unit(arena, &unit->cv_accel);
     }break;
     
     case SYMS_FileFormat_DWARF:
@@ -725,7 +728,7 @@ syms_type_sid_array_from_unit(SYMS_Arena *arena, SYMS_UnitAccel *unit){
   switch (unit->format){
     case SYMS_FileFormat_PDB:
     {
-      result = syms_pdb_type_sid_array_from_unit(arena, &unit->pdb_accel);
+      result = syms_cv_type_sid_array_from_unit(arena, &unit->cv_accel);
     }break;
     
     case SYMS_FileFormat_DWARF:
@@ -746,7 +749,7 @@ syms_symbol_kind_from_sid(SYMS_String8 data, SYMS_DbgAccel *dbg, SYMS_UnitAccel
     switch (unit->format){
       case SYMS_FileFormat_PDB:
       {
-        result = syms_pdb_symbol_kind_from_sid(data, &dbg->pdb_accel, &unit->pdb_accel, sid);
+        result = syms_pdb_symbol_kind_from_sid(data, &dbg->pdb_accel, &unit->cv_accel, sid);
       }break;
       
       case SYMS_FileFormat_DWARF:
@@ -769,7 +772,7 @@ syms_symbol_name_from_sid(SYMS_Arena *arena, SYMS_String8 data, SYMS_DbgAccel *d
     switch (unit->format){
       case SYMS_FileFormat_PDB:
       {
-        result = syms_pdb_symbol_name_from_sid(arena, data, &dbg->pdb_accel, &unit->pdb_accel, sid);
+        result = syms_cv_symbol_name_from_sid(arena, &unit->cv_accel, sid);
       }break;
       
       case SYMS_FileFormat_DWARF:
@@ -791,7 +794,7 @@ syms_type_info_from_sid(SYMS_String8 data, SYMS_DbgAccel *dbg, SYMS_UnitAccel *u
     switch (unit->format){
       case SYMS_FileFormat_PDB:
       {
-        result = syms_pdb_type_info_from_id(data, &dbg->pdb_accel, &unit->pdb_accel, id);
+        result = syms_pdb_type_info_from_sid(data, &dbg->pdb_accel, &unit->cv_accel, id);
       }break;
       
       case SYMS_FileFormat_DWARF:
@@ -813,7 +816,7 @@ syms_const_info_from_sid(SYMS_String8 data, SYMS_DbgAccel *dbg, SYMS_UnitAccel *
     switch (unit->format){
       case SYMS_FileFormat_PDB:
       {
-        result = syms_pdb_const_info_from_id(data, &dbg->pdb_accel, &unit->pdb_accel, id);
+        result = syms_pdb_const_info_from_id(data, &dbg->pdb_accel, &unit->cv_accel, id);
       }break;
       
       case SYMS_FileFormat_DWARF:
@@ -837,7 +840,7 @@ syms_type_from_var_sid(SYMS_String8 data, SYMS_DbgAccel *dbg, SYMS_UnitAccel *un
     switch (unit->format){
       case SYMS_FileFormat_PDB:
       {
-        result = syms_pdb_type_from_var_id(data, &dbg->pdb_accel, &unit->pdb_accel, id);
+        result = syms_pdb_type_from_var_id(data, &dbg->pdb_accel, &unit->cv_accel, id);
       }break;
       
       case SYMS_FileFormat_DWARF:
@@ -859,7 +862,7 @@ syms_voff_from_var_sid(SYMS_String8 data, SYMS_DbgAccel *dbg, SYMS_UnitAccel *un
     switch (unit->format){
       case SYMS_FileFormat_PDB:
       {
-        result = syms_pdb_voff_from_var_sid(data, &dbg->pdb_accel, &unit->pdb_accel, id);
+        result = syms_pdb_voff_from_var_sid(data, &dbg->pdb_accel, &unit->cv_accel, id);
       }break;
       
       case SYMS_FileFormat_DWARF:
@@ -882,7 +885,7 @@ syms_location_from_var_sid(SYMS_Arena *arena, SYMS_String8 data, SYMS_DbgAccel *
     switch (unit->format){
       case SYMS_FileFormat_PDB:
       {
-        result = syms_pdb_location_from_var_sid(arena, data, &dbg->pdb_accel, &unit->pdb_accel, sid);
+        result = syms_pdb_location_from_var_sid(arena, data, &dbg->pdb_accel, &unit->cv_accel, sid);
       }break;
       
       case SYMS_FileFormat_DWARF:
@@ -906,7 +909,7 @@ syms_location_ranges_from_var_sid(SYMS_Arena *arena, SYMS_String8 data, SYMS_Dbg
       case SYMS_FileFormat_PDB:
       {
         result = syms_pdb_location_ranges_from_var_sid(arena, data, &dbg->pdb_accel,
-                                                       &unit->pdb_accel, sid);
+                                                       &unit->cv_accel, sid);
       }break;
       
       case SYMS_FileFormat_DWARF:
@@ -929,7 +932,7 @@ syms_location_from_id(SYMS_Arena *arena, SYMS_String8 data, SYMS_DbgAccel *dbg,
     switch (unit->format){
       case SYMS_FileFormat_PDB:
       {
-        result = syms_pdb_location_from_id(arena, data, &dbg->pdb_accel, &unit->pdb_accel, loc_id);
+        result = syms_pdb_location_from_id(arena, data, &dbg->pdb_accel, &unit->cv_accel, loc_id);
       }break;
       
       case SYMS_FileFormat_DWARF:
@@ -955,7 +958,7 @@ syms_mems_accel_from_sid(SYMS_Arena *arena, SYMS_String8 data, SYMS_DbgAccel *db
     switch (unit->format){
       case SYMS_FileFormat_PDB:
       {
-        result = (SYMS_MemsAccel*)syms_pdb_mems_accel_from_sid(arena, data, &dbg->pdb_accel, &unit->pdb_accel, id);
+        result = (SYMS_MemsAccel*)syms_pdb_mems_accel_from_sid(arena, data, &dbg->pdb_accel, &unit->cv_accel, id);
       }break;
       
       case SYMS_FileFormat_DWARF:
@@ -976,7 +979,7 @@ syms_mem_count_from_mems(SYMS_MemsAccel *mems){
   switch (mems->format){
     case SYMS_FileFormat_PDB:
     {
-      result = syms_pdb_mem_count_from_mems(&mems->pdb_accel);
+      result = syms_pdb_mem_count_from_mems(&mems->cv_accel);
     }break;
     
     case SYMS_FileFormat_DWARF:
@@ -998,7 +1001,7 @@ syms_mem_info_from_number(SYMS_Arena *arena, SYMS_String8 data, SYMS_DbgAccel *d
     switch (mems->format){
       case SYMS_FileFormat_PDB:
       {
-        result = syms_pdb_mem_info_from_number(arena, data, &dbg->pdb_accel, &unit->pdb_accel, &mems->pdb_accel, n);
+        result = syms_pdb_mem_info_from_number(arena, data, &dbg->pdb_accel, &unit->cv_accel, &mems->cv_accel, n);
       }break;
       
       case SYMS_FileFormat_DWARF:
@@ -1021,7 +1024,7 @@ syms_type_from_mem_number(SYMS_String8 data, SYMS_DbgAccel *dbg, SYMS_UnitAccel
     switch (mems->format){
       case SYMS_FileFormat_PDB:
       {
-        result = syms_pdb_type_from_mem_number(data, &dbg->pdb_accel, &unit->pdb_accel, &mems->pdb_accel, n);
+        result = syms_pdb_type_from_mem_number(data, &dbg->pdb_accel, &unit->cv_accel, &mems->cv_accel, n);
       }break;
       
       case SYMS_FileFormat_DWARF:
@@ -1044,8 +1047,8 @@ syms_sig_info_from_mem_number(SYMS_Arena *arena, SYMS_String8 data, SYMS_DbgAcce
     switch (mems->format){
       case SYMS_FileFormat_PDB:
       {
-        result = syms_pdb_sig_info_from_mem_number(arena, data, &dbg->pdb_accel, &unit->pdb_accel,
-                                                   &mems->pdb_accel, n);
+        result = syms_pdb_sig_info_from_mem_number(arena, data, &dbg->pdb_accel, &unit->cv_accel,
+                                                   &mems->cv_accel, n);
       }break;
       
       case SYMS_FileFormat_DWARF:
@@ -1104,21 +1107,43 @@ syms_containing_type_from_sid(SYMS_String8 data, SYMS_DbgAccel *dbg, SYMS_UnitAc
   return(result);
 }
 
-SYMS_API SYMS_EnumInfoArray
-syms_enum_info_array_from_sid(SYMS_Arena *arena, SYMS_String8 data, SYMS_DbgAccel *dbg,
-                              SYMS_UnitAccel *unit, SYMS_SymbolID sid){
-  SYMS_ProfBegin("syms_enum_info_array_from_sid");
-  SYMS_EnumInfoArray result = {0};
+SYMS_API SYMS_String8
+syms_linkage_name_from_sid(SYMS_Arena *arena, SYMS_String8 data, SYMS_DbgAccel *dbg, SYMS_UnitAccel *unit, SYMS_SymbolID sid){
+  SYMS_ProfBegin("syms_linkage_name_from_sid");
+  SYMS_String8 result = {0};
+  if(dbg->format == unit->format){
+    switch(unit->format){
+      case SYMS_FileFormat_PDB:
+      {
+        // NOTE(rjf): not available in PDB
+      }break;
+      
+      case SYMS_FileFormat_DWARF:
+      {
+        result = syms_dw_linkage_name_from_sid(arena, data, &dbg->dw_accel, &unit->dw_accel, sid);
+      }break;
+    }
+  }
+  SYMS_ASSERT_PARANOID(syms_parser_api_invariants());
+  SYMS_ProfEnd();
+  return result;
+}
+
+SYMS_API SYMS_EnumMemberArray
+syms_enum_member_array_from_sid(SYMS_Arena *arena, SYMS_String8 data, SYMS_DbgAccel *dbg,
+                                SYMS_UnitAccel *unit, SYMS_SymbolID sid){
+  SYMS_ProfBegin("syms_enum_member_array_from_sid");
+  SYMS_EnumMemberArray result = {0};
   if (dbg->format == unit->format){
     switch (unit->format){
       case SYMS_FileFormat_PDB:
       {
-        result = syms_pdb_enum_info_array_from_sid(arena, data, &dbg->pdb_accel, &unit->pdb_accel, sid);
+        result = syms_pdb_enum_member_array_from_sid(arena, data, &dbg->pdb_accel, &unit->cv_accel, sid);
       }break;
       
       case SYMS_FileFormat_DWARF:
       {
-        result = syms_dw_enum_info_array_from_sid(arena, data, &dbg->dw_accel, &unit->dw_accel, sid);
+        result = syms_dw_enum_member_array_from_sid(arena, data, &dbg->dw_accel, &unit->dw_accel, sid);
       }break;
     }
   }
@@ -1137,7 +1162,7 @@ syms_sig_handle_from_proc_sid(SYMS_String8 data, SYMS_DbgAccel *dbg, SYMS_UnitAc
     switch (unit->format){
       case SYMS_FileFormat_PDB:
       {
-        result = syms_pdb_proc_sig_handle_from_id(data, &dbg->pdb_accel, &unit->pdb_accel, sid);
+        result = syms_pdb_proc_sig_handle_from_id(data, &dbg->pdb_accel, &unit->cv_accel, sid);
       }break;
       
       case SYMS_FileFormat_DWARF:
@@ -1160,7 +1185,7 @@ syms_sig_info_from_handle(SYMS_Arena *arena, SYMS_String8 data, SYMS_DbgAccel *d
     switch (unit->format){
       case SYMS_FileFormat_PDB:
       {
-        result = syms_pdb_sig_info_from_handle(arena, data, &dbg->pdb_accel, &unit->pdb_accel, handle);
+        result = syms_pdb_sig_info_from_handle(arena, data, &dbg->pdb_accel, &unit->cv_accel, handle);
       }break;
       
       case SYMS_FileFormat_DWARF:
@@ -1183,7 +1208,7 @@ syms_scope_vranges_from_sid(SYMS_Arena *arena, SYMS_String8 data, SYMS_DbgAccel
     switch (unit->format){
       case SYMS_FileFormat_PDB:
       {
-        result = syms_pdb_scope_vranges_from_sid(arena, data, &dbg->pdb_accel, &unit->pdb_accel, id);
+        result = syms_pdb_scope_vranges_from_sid(arena, data, &dbg->pdb_accel, &unit->cv_accel, id);
       }break;
       
       case SYMS_FileFormat_DWARF:
@@ -1206,7 +1231,7 @@ syms_scope_children_from_sid(SYMS_Arena *arena, SYMS_String8 data, SYMS_DbgAccel
     switch (unit->format){
       case SYMS_FileFormat_PDB:
       {
-        result = syms_pdb_scope_children_from_sid(arena, data, &dbg->pdb_accel, &unit->pdb_accel, sid);
+        result = syms_pdb_scope_children_from_sid(arena, data, &dbg->pdb_accel, &unit->cv_accel, sid);
       }break;
       
       case SYMS_FileFormat_DWARF:
@@ -1276,7 +1301,7 @@ syms_sig_info_from_type_sid(SYMS_Arena *arena, SYMS_String8 data, SYMS_DbgAccel
     switch (unit->format){
       case SYMS_FileFormat_PDB:
       {
-        result = syms_pdb_sig_info_from_id(arena, data, &dbg->pdb_accel, &unit->pdb_accel, sid);
+        result = syms_pdb_sig_info_from_id(arena, data, &dbg->pdb_accel, &unit->cv_accel, sid);
       }break;
       
       case SYMS_FileFormat_DWARF:
@@ -1429,7 +1454,7 @@ syms_usid_list_from_string(SYMS_Arena *arena, SYMS_String8 data, SYMS_DbgAccel *
       case SYMS_FileFormat_PDB:
       {
         result = syms_pdb_usid_list_from_string(arena, data, &dbg->pdb_accel,
-                                                &map_and_unit->unit->pdb_accel,
+                                                &map_and_unit->unit->cv_accel,
                                                 &map_and_unit->map->pdb_accel,
                                                 string);
       }break;
@@ -1511,7 +1536,7 @@ syms_voff_from_link_name(SYMS_String8 data, SYMS_DbgAccel *dbg, SYMS_LinkMapAcce
       case SYMS_FileFormat_PDB:
       {
         result = syms_pdb_voff_from_link_name(data, (SYMS_PdbDbgAccel*)dbg, (SYMS_PdbLinkMapAccel*)map,
-                                              (SYMS_PdbUnitAccel*)link_unit, name);
+                                              (SYMS_CvUnitAccel*)link_unit, name);
       }break;
       
       case SYMS_FileFormat_DWARF:
@@ -1532,7 +1557,7 @@ syms_link_name_array_from_unit(SYMS_Arena *arena, SYMS_String8 data, SYMS_DbgAcc
     switch (dbg->format){
       case SYMS_FileFormat_PDB:
       {
-        result = syms_pdb_link_name_array_from_unit(arena, data, (SYMS_PdbDbgAccel*)dbg, (SYMS_PdbUnitAccel*)unit);
+        result = syms_pdb_link_name_array_from_unit(arena, data, (SYMS_PdbDbgAccel*)dbg, (SYMS_CvUnitAccel*)unit);
       }break;
       
       case SYMS_FileFormat_DWARF:
@@ -1566,7 +1591,7 @@ syms_tls_var_sid_array_from_unit(SYMS_Arena *arena, SYMS_UnitAccel *unit){
   switch (unit->format){
     case SYMS_FileFormat_PDB:
     {
-      result = syms_pdb_tls_var_sid_array_from_unit(arena, &unit->pdb_accel);
+      result = syms_pdb_tls_var_sid_array_from_unit(arena, &unit->cv_accel);
     }break;
     
     case SYMS_FileFormat_DWARF:
@@ -1579,5 +1604,4 @@ syms_tls_var_sid_array_from_unit(SYMS_Arena *arena, SYMS_UnitAccel *unit){
   return(result);
 }
 
-
 #endif //SYMS_PARSER_C
diff --git a/Engine/Source/Runtime/SymsLib/syms/core/syms_parser.h b/Engine/Source/Runtime/SymsLib/syms/core/syms_parser.h
--- a/Engine/Source/Runtime/SymsLib/syms/core/syms_parser.h
+++ b/Engine/Source/Runtime/SymsLib/syms/core/syms_parser.h
@@ -11,7 +11,6 @@ typedef union SYMS_FileAccel{
   SYMS_FileFormat format;
   SYMS_ElfFileAccel elf_accel;
   SYMS_PeFileAccel pe_accel;
-  SYMS_CoffFileAccel coff_accel;
   SYMS_PdbFileAccel pdb_accel;
   SYMS_MachFileAccel mach_accel;
 } SYMS_FileAccel;
@@ -19,7 +18,6 @@ typedef union SYMS_FileAccel{
 typedef union SYMS_BinAccel{
   SYMS_FileFormat format;
   SYMS_PeBinAccel pe_accel;
-  SYMS_CoffBinAccel coff_accel;
   SYMS_ElfBinAccel elf_accel;
   SYMS_MachBinAccel mach_accel;
 } SYMS_BinAccel;
@@ -44,13 +42,13 @@ typedef union SYMS_UnitSetAccel{
 typedef union SYMS_UnitAccel{
   SYMS_FileFormat format;
   SYMS_DwUnitAccel dw_accel;
-  SYMS_PdbUnitAccel pdb_accel;
+  SYMS_CvUnitAccel cv_accel;
 } SYMS_UnitAccel;
 
 typedef union SYMS_MemsAccel{
   SYMS_FileFormat format;
   SYMS_DwMemsAccel dw_accel;
-  SYMS_PdbMemsAccel pdb_accel;
+  SYMS_CvMemsAccel cv_accel;
 } SYMS_MemsAccel;
 
 typedef union SYMS_MapAccel{
@@ -98,29 +96,38 @@ SYMS_API SYMS_FileFormat   syms_file_format_from_file(SYMS_FileAccel *file);
 //~ allen: Bin File
 
 SYMS_API SYMS_B32          syms_file_is_bin(SYMS_FileAccel *file);
-SYMS_API SYMS_BinAccel*    syms_bin_accel_from_file(SYMS_Arena *arena, SYMS_String8 data, SYMS_FileAccel *file);
+SYMS_API SYMS_BinAccel*    syms_bin_accel_from_file(SYMS_Arena *arena, SYMS_String8 data,
+                                                    SYMS_FileAccel *file);
 
 // arch
 SYMS_API SYMS_Arch         syms_arch_from_bin(SYMS_BinAccel *bin);
 
 // external info
-SYMS_API SYMS_ExtFileList  syms_ext_file_list_from_bin(SYMS_Arena *arena, SYMS_String8 data, SYMS_BinAccel *bin);
+SYMS_API SYMS_ExtFileList  syms_ext_file_list_from_bin(SYMS_Arena *arena, SYMS_String8 data,
+                                                       SYMS_BinAccel *bin);
 
 // binary secs
-SYMS_API SYMS_SecInfoArray syms_sec_info_array_from_bin(SYMS_Arena *arena, SYMS_String8 data, SYMS_BinAccel *bin);
+SYMS_API SYMS_SecInfoArray syms_sec_info_array_from_bin(SYMS_Arena *arena, SYMS_String8 data,
+                                                        SYMS_BinAccel *bin);
 
 // default vbase
 SYMS_API SYMS_U64          syms_default_vbase_from_bin(SYMS_BinAccel *bin);
 
+// entry point
+SYMS_API SYMS_U64          syms_entry_point_voff_from_bin(SYMS_BinAccel *bin);
+
 // imports & exports
-SYMS_API SYMS_ImportArray  syms_imports_from_bin(SYMS_Arena *arena, SYMS_String8 data, SYMS_BinAccel *bin);
-SYMS_API SYMS_ExportArray  syms_exports_from_bin(SYMS_Arena *arena, SYMS_String8 data, SYMS_BinAccel *bin);
+SYMS_API SYMS_ImportArray  syms_imports_from_bin(SYMS_Arena *arena, SYMS_String8 data,
+                                                 SYMS_BinAccel *bin);
+SYMS_API SYMS_ExportArray  syms_exports_from_bin(SYMS_Arena *arena, SYMS_String8 data,
+                                                 SYMS_BinAccel *bin);
 
 ////////////////////////////////
 //~ nick: Bin List
 
 SYMS_API SYMS_B32          syms_file_is_bin_list(SYMS_FileAccel *file);
-SYMS_API SYMS_BinListAccel*syms_bin_list_from_file(SYMS_Arena *arena, SYMS_String8 data, SYMS_FileAccel *file);
+SYMS_API SYMS_BinListAccel*syms_bin_list_from_file(SYMS_Arena *arena, SYMS_String8 data,
+                                                   SYMS_FileAccel *file);
 SYMS_API SYMS_BinInfoArray syms_bin_info_array_from_bin_list(SYMS_Arena *arena, SYMS_BinListAccel *list);
 SYMS_API SYMS_BinAccel*    syms_bin_accel_from_bin_list_number(SYMS_Arena *arena, SYMS_String8 data,
                                                                SYMS_BinListAccel *list, SYMS_U64 n);
@@ -129,42 +136,48 @@ SYMS_API SYMS_BinAccel*    syms_bin_accel_from_bin_list_number(SYMS_Arena *arena
 //~ rjf: Dbg File
 
 SYMS_API SYMS_B32          syms_file_is_dbg(SYMS_FileAccel *file);
-SYMS_API SYMS_DbgAccel*    syms_dbg_accel_from_file(SYMS_Arena *arena, SYMS_String8 data, SYMS_FileAccel *file);
+SYMS_API SYMS_DbgAccel*    syms_dbg_accel_from_file(SYMS_Arena *arena, SYMS_String8 data,
+                                                    SYMS_FileAccel *file);
 
 SYMS_API SYMS_B32          syms_bin_is_dbg(SYMS_BinAccel *bin);
-SYMS_API SYMS_DbgAccel*    syms_dbg_accel_from_bin(SYMS_Arena *arena, SYMS_String8 data, SYMS_BinAccel *bin);
+SYMS_API SYMS_DbgAccel*    syms_dbg_accel_from_bin(SYMS_Arena *arena, SYMS_String8 data,
+                                                   SYMS_BinAccel *bin);
 
 // arch
 SYMS_API SYMS_Arch         syms_arch_from_dbg(SYMS_DbgAccel *dbg);
 
 // external info
-SYMS_API SYMS_ExtFileList  syms_ext_file_list_from_dbg(SYMS_Arena *arena, SYMS_String8 data, SYMS_DbgAccel *dbg);
+SYMS_API SYMS_ExtFileList  syms_ext_file_list_from_dbg(SYMS_Arena *arena, SYMS_String8 data,
+                                                       SYMS_DbgAccel *dbg);
 
 // match key
 SYMS_API SYMS_ExtMatchKey  syms_ext_match_key_from_dbg(SYMS_String8 data, SYMS_DbgAccel *dbg);
 
 // binary secs
-SYMS_API SYMS_SecInfoArray syms_sec_info_array_from_dbg(SYMS_Arena *arena, SYMS_String8 data, SYMS_DbgAccel *dbg);
+SYMS_API SYMS_SecInfoArray syms_sec_info_array_from_dbg(SYMS_Arena *arena, SYMS_String8 data,
+                                                        SYMS_DbgAccel *dbg);
 
 // default vbase
 SYMS_API SYMS_U64          syms_default_vbase_from_dbg(SYMS_DbgAccel *dbg);
 
 // compilation units
-SYMS_API SYMS_UnitSetAccel*syms_unit_set_accel_from_dbg(SYMS_Arena *arena, SYMS_String8 data, SYMS_DbgAccel *dbg);
+SYMS_API SYMS_UnitSetAccel*syms_unit_set_accel_from_dbg(SYMS_Arena *arena, SYMS_String8 data,
+                                                        SYMS_DbgAccel *dbg);
 SYMS_API SYMS_U64          syms_unit_count_from_set(SYMS_UnitSetAccel *unit_set);
 SYMS_API SYMS_UnitInfo     syms_unit_info_from_uid(SYMS_UnitSetAccel *unit_set, SYMS_UnitID uid);
 SYMS_API SYMS_UnitNames    syms_unit_names_from_uid(SYMS_Arena *arena, SYMS_UnitSetAccel *unit_set,
                                                     SYMS_UnitID uid);
 
-SYMS_API SYMS_UnitRangeArray syms_unit_ranges_from_set(SYMS_Arena *arena, SYMS_String8 data, SYMS_DbgAccel *dbg,
-                                                       SYMS_UnitSetAccel *unit_set);
+SYMS_API SYMS_UnitRangeArray syms_unit_ranges_from_set(SYMS_Arena *arena, SYMS_String8 data,
+                                                       SYMS_DbgAccel *dbg, SYMS_UnitSetAccel *unit_set);
 
 SYMS_API SYMS_UnitID       syms_uid_collated_types_from_set(SYMS_UnitSetAccel *unit_set);
 SYMS_API SYMS_UnitID       syms_uid_collated_public_symbols_from_set(SYMS_UnitSetAccel *unit_set);
 
 // symbol parsing
-SYMS_API SYMS_UnitAccel*    syms_unit_accel_from_uid(SYMS_Arena *arena, SYMS_String8 data, SYMS_DbgAccel *dbg,
-                                                     SYMS_UnitSetAccel *unit_set, SYMS_UnitID uid);
+SYMS_API SYMS_UnitAccel*    syms_unit_accel_from_uid(SYMS_Arena *arena, SYMS_String8 data,
+                                                     SYMS_DbgAccel *dbg, SYMS_UnitSetAccel *unit_set,
+                                                     SYMS_UnitID uid);
 
 SYMS_API SYMS_UnitID        syms_uid_from_unit(SYMS_UnitAccel *unit);
 
@@ -175,106 +188,128 @@ SYMS_API SYMS_SymbolIDArray syms_type_sid_array_from_unit(SYMS_Arena *arena, SYM
 
 SYMS_API SYMS_SymbolKind    syms_symbol_kind_from_sid(SYMS_String8 data, SYMS_DbgAccel *dbg,
                                                       SYMS_UnitAccel *unit, SYMS_SymbolID sid);
-SYMS_API SYMS_String8       syms_symbol_name_from_sid(SYMS_Arena *arena, SYMS_String8 data, SYMS_DbgAccel *dbg,
-                                                      SYMS_UnitAccel *unit, SYMS_SymbolID sid);
+SYMS_API SYMS_String8       syms_symbol_name_from_sid(SYMS_Arena *arena, SYMS_String8 data,
+                                                      SYMS_DbgAccel *dbg, SYMS_UnitAccel *unit,
+                                                      SYMS_SymbolID sid);
 
 
-SYMS_API SYMS_TypeInfo      syms_type_info_from_sid(SYMS_String8 data, SYMS_DbgAccel *dbg, SYMS_UnitAccel *unit,
-                                                    SYMS_SymbolID sid);
+SYMS_API SYMS_TypeInfo      syms_type_info_from_sid(SYMS_String8 data, SYMS_DbgAccel *dbg,
+                                                    SYMS_UnitAccel *unit, SYMS_SymbolID sid);
 
-SYMS_API SYMS_ConstInfo     syms_const_info_from_sid(SYMS_String8 data, SYMS_DbgAccel *dbg, SYMS_UnitAccel *unit,
-                                                     SYMS_SymbolID sid);
+SYMS_API SYMS_ConstInfo     syms_const_info_from_sid(SYMS_String8 data, SYMS_DbgAccel *dbg,
+                                                     SYMS_UnitAccel *unit, SYMS_SymbolID sid);
 
 // variable info
-SYMS_API SYMS_USID         syms_type_from_var_sid(SYMS_String8 data, SYMS_DbgAccel *dbg, SYMS_UnitAccel *unit,
-                                                  SYMS_SymbolID sid);
-
-SYMS_API SYMS_U64          syms_voff_from_var_sid(SYMS_String8 data, SYMS_DbgAccel *dbg, SYMS_UnitAccel *unit,
-                                                  SYMS_SymbolID sid);
-
-SYMS_API SYMS_Location     syms_location_from_var_sid(SYMS_Arena *arena, SYMS_String8 data, SYMS_DbgAccel *dbg,
-                                                      SYMS_UnitAccel *unit, SYMS_SymbolID sid);
-SYMS_API SYMS_LocRangeArray syms_location_ranges_from_var_sid(SYMS_Arena *arena, SYMS_String8 data, SYMS_DbgAccel *dbg,
-                                                              SYMS_UnitAccel *unit, SYMS_SymbolID sid);
+SYMS_API SYMS_USID         syms_type_from_var_sid(SYMS_String8 data, SYMS_DbgAccel *dbg,
+                                                  SYMS_UnitAccel *unit, SYMS_SymbolID sid);
+
+SYMS_API SYMS_U64          syms_voff_from_var_sid(SYMS_String8 data, SYMS_DbgAccel *dbg,
+                                                  SYMS_UnitAccel *unit, SYMS_SymbolID sid);
+
+SYMS_API SYMS_Location     syms_location_from_var_sid(SYMS_Arena *arena, SYMS_String8 data,
+                                                      SYMS_DbgAccel *dbg, SYMS_UnitAccel *unit,
+                                                      SYMS_SymbolID sid);
+SYMS_API SYMS_LocRangeArray syms_location_ranges_from_var_sid(SYMS_Arena *arena, SYMS_String8 data,
+                                                              SYMS_DbgAccel *dbg, SYMS_UnitAccel *unit,
+                                                              SYMS_SymbolID sid);
 SYMS_API SYMS_Location     syms_location_from_id(SYMS_Arena *arena, SYMS_String8 data, SYMS_DbgAccel *dbg,
                                                  SYMS_UnitAccel *unit, SYMS_LocID loc_id);
 
 // member info
-SYMS_API SYMS_MemsAccel*   syms_mems_accel_from_sid(SYMS_Arena *arena, SYMS_String8 data, SYMS_DbgAccel *dbg,
-                                                    SYMS_UnitAccel *unit, SYMS_SymbolID sid);
+SYMS_API SYMS_MemsAccel*   syms_mems_accel_from_sid(SYMS_Arena *arena, SYMS_String8 data,
+                                                    SYMS_DbgAccel *dbg, SYMS_UnitAccel *unit,
+                                                    SYMS_SymbolID sid);
 
 SYMS_API SYMS_U64          syms_mem_count_from_mems(SYMS_MemsAccel *mems);
-SYMS_API SYMS_MemInfo      syms_mem_info_from_number(SYMS_Arena *arena, SYMS_String8 data, SYMS_DbgAccel *dbg,
-                                                     SYMS_UnitAccel *unit, SYMS_MemsAccel *mems, SYMS_U64 n);
-
-SYMS_API SYMS_USID         syms_type_from_mem_number(SYMS_String8 data, SYMS_DbgAccel *dbg, SYMS_UnitAccel *unit,
+SYMS_API SYMS_MemInfo      syms_mem_info_from_number(SYMS_Arena *arena, SYMS_String8 data,
+                                                     SYMS_DbgAccel *dbg, SYMS_UnitAccel *unit,
                                                      SYMS_MemsAccel *mems, SYMS_U64 n);
-SYMS_API SYMS_SigInfo      syms_sig_info_from_mem_number(SYMS_Arena *arena, SYMS_String8 data, SYMS_DbgAccel *dbg,
-                                                         SYMS_UnitAccel *unit, SYMS_MemsAccel *mems, SYMS_U64 n);
 
-SYMS_API SYMS_USID         syms_symbol_from_mem_number(SYMS_String8 data, SYMS_DbgAccel *dbg, SYMS_UnitAccel *unit, SYMS_MemsAccel *mems, SYMS_U64 n);
+SYMS_API SYMS_USID         syms_type_from_mem_number(SYMS_String8 data, SYMS_DbgAccel *dbg,
+                                                     SYMS_UnitAccel *unit, SYMS_MemsAccel *mems,
+                                                     SYMS_U64 n);
+SYMS_API SYMS_SigInfo      syms_sig_info_from_mem_number(SYMS_Arena *arena, SYMS_String8 data,
+                                                         SYMS_DbgAccel *dbg, SYMS_UnitAccel *unit,
+                                                         SYMS_MemsAccel *mems, SYMS_U64 n);
+
+SYMS_API SYMS_USID         syms_symbol_from_mem_number(SYMS_String8 data, SYMS_DbgAccel *dbg,
+                                                       SYMS_UnitAccel *unit, SYMS_MemsAccel *mems,
+                                                       SYMS_U64 n);
 SYMS_API SYMS_USID         syms_containing_type_from_sid(SYMS_String8 data, SYMS_DbgAccel *dbg,
                                                          SYMS_UnitAccel *unit, SYMS_SymbolID sid);
-// TODO(rjf): linkage name from symbol
 
-SYMS_API SYMS_EnumInfoArray syms_enum_info_array_from_sid(SYMS_Arena *arena, SYMS_String8 data, SYMS_DbgAccel *dbg,
-                                                          SYMS_UnitAccel *unit, SYMS_SymbolID sid);
+SYMS_API SYMS_String8      syms_linkage_name_from_sid(SYMS_Arena *arena, SYMS_String8 data,
+                                                      SYMS_DbgAccel *dbg, SYMS_UnitAccel *unit,
+                                                      SYMS_SymbolID sid);
+
+SYMS_API SYMS_EnumMemberArray syms_enum_member_array_from_sid(SYMS_Arena *arena, SYMS_String8 data,
+                                                              SYMS_DbgAccel *dbg, SYMS_UnitAccel *unit,
+                                                              SYMS_SymbolID sid);
 
 // procedure info
 SYMS_API SYMS_UnitIDAndSig syms_sig_handle_from_proc_sid(SYMS_String8 data, SYMS_DbgAccel *dbg,
                                                          SYMS_UnitAccel *unit, SYMS_SymbolID proc_sid);
-SYMS_API SYMS_SigInfo      syms_sig_info_from_handle(SYMS_Arena *arena, SYMS_String8 data, SYMS_DbgAccel *dbg,
-                                                     SYMS_UnitAccel *unit, SYMS_SigHandle handle);
+SYMS_API SYMS_SigInfo      syms_sig_info_from_handle(SYMS_Arena *arena, SYMS_String8 data,
+                                                     SYMS_DbgAccel *dbg, SYMS_UnitAccel *unit,
+                                                     SYMS_SigHandle handle);
 
 #define syms_proc_vranges_from_sid syms_scope_vranges_from_sid
-SYMS_API SYMS_U64RangeArray syms_scope_vranges_from_sid(SYMS_Arena *arena, SYMS_String8 data, SYMS_DbgAccel *dbg,
-                                                        SYMS_UnitAccel *unit, SYMS_SymbolID sid);
+SYMS_API SYMS_U64RangeArray syms_scope_vranges_from_sid(SYMS_Arena *arena, SYMS_String8 data,
+                                                        SYMS_DbgAccel *dbg, SYMS_UnitAccel *unit,
+                                                        SYMS_SymbolID sid);
 SYMS_API SYMS_SymbolIDArray syms_scope_children_from_sid(SYMS_Arena *arena, SYMS_String8 data,
                                                          SYMS_DbgAccel *dbg, SYMS_UnitAccel *unit,
                                                          SYMS_SymbolID sid);
 
-SYMS_API SYMS_Location     syms_location_from_proc_sid(SYMS_Arena *arena, SYMS_String8 data, SYMS_DbgAccel *dbg,
-                                                       SYMS_UnitAccel *unit, SYMS_SymbolID sid,
-                                                       SYMS_ProcLoc proc_loc);
+SYMS_API SYMS_Location     syms_location_from_proc_sid(SYMS_Arena *arena, SYMS_String8 data,
+                                                       SYMS_DbgAccel *dbg, SYMS_UnitAccel *unit,
+                                                       SYMS_SymbolID sid, SYMS_ProcLoc proc_loc);
 SYMS_API SYMS_LocRangeArray syms_location_ranges_from_proc_sid(SYMS_Arena *arena, SYMS_String8 data,
-                                                               SYMS_DbgAccel *dbg,
-                                                               SYMS_UnitAccel *unit, SYMS_SymbolID sid,
-                                                               SYMS_ProcLoc proc_loc);
+                                                               SYMS_DbgAccel *dbg, SYMS_UnitAccel *unit,
+                                                               SYMS_SymbolID sid, SYMS_ProcLoc proc_loc);
 
 // signature info
-SYMS_API SYMS_SigInfo      syms_sig_info_from_type_sid(SYMS_Arena *arena, SYMS_String8 data, SYMS_DbgAccel *dbg,
-                                                       SYMS_UnitAccel *unit, SYMS_SymbolID sid);
+SYMS_API SYMS_SigInfo      syms_sig_info_from_type_sid(SYMS_Arena *arena, SYMS_String8 data,
+                                                       SYMS_DbgAccel *dbg, SYMS_UnitAccel *unit,
+                                                       SYMS_SymbolID sid);
 
 // line info
-SYMS_API SYMS_String8      syms_file_name_from_id(SYMS_Arena *arena, SYMS_String8 data, SYMS_DbgAccel *dbg,
-                                                  SYMS_UnitSetAccel *unit_set, SYMS_UnitID uid,
-                                                  SYMS_FileID file_id);
-SYMS_API SYMS_String8Array syms_file_table_from_uid(SYMS_Arena *arena, SYMS_String8 data, SYMS_DbgAccel *dbg,
-                                                    SYMS_UnitSetAccel *unit_set, SYMS_UnitID uid);
-SYMS_API SYMS_LineParseOut syms_line_parse_from_uid(SYMS_Arena *arena, SYMS_String8 data, SYMS_DbgAccel *dbg,
-                                                    SYMS_UnitSetAccel *unit_set, SYMS_UnitID uid);
+SYMS_API SYMS_String8      syms_file_name_from_id(SYMS_Arena *arena, SYMS_String8 data,
+                                                  SYMS_DbgAccel *dbg, SYMS_UnitSetAccel *unit_set,
+                                                  SYMS_UnitID uid, SYMS_FileID file_id);
+SYMS_API SYMS_String8Array syms_file_table_from_uid(SYMS_Arena *arena, SYMS_String8 data,
+                                                    SYMS_DbgAccel *dbg, SYMS_UnitSetAccel *unit_set,
+                                                    SYMS_UnitID uid);
+SYMS_API SYMS_LineParseOut syms_line_parse_from_uid(SYMS_Arena *arena, SYMS_String8 data,
+                                                    SYMS_DbgAccel *dbg, SYMS_UnitSetAccel *unit_set,
+                                                    SYMS_UnitID uid);
 
 
 // name maps
-SYMS_API SYMS_MapAccel*    syms_type_map_from_dbg(SYMS_Arena *arena, SYMS_String8 data, SYMS_DbgAccel *dbg);
+SYMS_API SYMS_MapAccel*    syms_type_map_from_dbg(SYMS_Arena *arena, SYMS_String8 data,
+                                                  SYMS_DbgAccel *dbg);
 SYMS_API SYMS_MapAccel*    syms_unmangled_symbol_map_from_dbg(SYMS_Arena *arena,SYMS_String8 data,SYMS_DbgAccel *dbg);
 SYMS_API SYMS_UnitID       syms_partner_uid_from_map(SYMS_MapAccel *map);
-SYMS_API SYMS_USIDList     syms_usid_list_from_string(SYMS_Arena *arena, SYMS_String8 data, SYMS_DbgAccel *dbg,
-                                                      SYMS_MapAndUnit *map_and_unit, SYMS_String8 string);
+SYMS_API SYMS_USIDList     syms_usid_list_from_string(SYMS_Arena *arena, SYMS_String8 data,
+                                                      SYMS_DbgAccel *dbg, SYMS_MapAndUnit *map_and_unit,
+                                                      SYMS_String8 string);
 
 // linker names
 SYMS_API SYMS_UnitID       syms_link_names_uid(SYMS_DbgAccel *dbg);
 
-SYMS_API SYMS_LinkMapAccel*syms_link_map_from_dbg(SYMS_Arena *arena, SYMS_String8 data, SYMS_DbgAccel *dbg);
+SYMS_API SYMS_LinkMapAccel*syms_link_map_from_dbg(SYMS_Arena *arena, SYMS_String8 data,
+                                                  SYMS_DbgAccel *dbg);
 SYMS_API SYMS_B32          syms_link_map_is_complete(SYMS_LinkMapAccel *map);
-SYMS_API SYMS_U64          syms_voff_from_link_name(SYMS_String8 data, SYMS_DbgAccel *dbg, SYMS_LinkMapAccel *map,
-                                                    SYMS_UnitAccel *link_unit, SYMS_String8 name);
+SYMS_API SYMS_U64          syms_voff_from_link_name(SYMS_String8 data, SYMS_DbgAccel *dbg,
+                                                    SYMS_LinkMapAccel *map, SYMS_UnitAccel *link_unit,
+                                                    SYMS_String8 name);
 SYMS_API SYMS_LinkNameRecArray syms_link_name_array_from_unit(SYMS_Arena *arena, SYMS_String8 data,
                                                               SYMS_DbgAccel *dbg, SYMS_UnitAccel *unit);
 
 // thread vars
 SYMS_API SYMS_UnitID        syms_tls_var_uid_from_dbg(SYMS_DbgAccel *dbg);
-SYMS_API SYMS_SymbolIDArray syms_tls_var_sid_array_from_unit(SYMS_Arena *arena, SYMS_UnitAccel *thread_unit);
+SYMS_API SYMS_SymbolIDArray syms_tls_var_sid_array_from_unit(SYMS_Arena *arena,
+                                                             SYMS_UnitAccel *thread_unit);
 
 SYMS_C_LINKAGE_END
 
diff --git a/Engine/Source/Runtime/SymsLib/syms/core/unwind/syms_unwind_pe_x64.c b/Engine/Source/Runtime/SymsLib/syms/core/unwind/syms_unwind_pe_x64.c
--- a/Engine/Source/Runtime/SymsLib/syms/core/unwind/syms_unwind_pe_x64.c
+++ b/Engine/Source/Runtime/SymsLib/syms/core/unwind/syms_unwind_pe_x64.c
@@ -210,7 +210,7 @@ syms_unwind_pe_x64(SYMS_String8 data, SYMS_PeBinAccel *bin, SYMS_U64 bin_base,
               {
                 // NOTE(rjf): this was found by stepping through kernel code after an exception was
                 // thrown, encountered in the exception_stepping_tests (after the throw) in mule_main
-                SYMS_ASSERT(!"Hit me!");
+                // SYMS_ASSERT(!"Hit me!");
                 result.dead = syms_true;
                 // TODO(allen): ???
               }break;
@@ -257,7 +257,7 @@ syms_unwind_pe_x64(SYMS_String8 data, SYMS_PeBinAccel *bin, SYMS_U64 bin_base,
               {
                 // NOTE(rjf): this was found by stepping through kernel code after an exception was
                 // thrown, encountered in the exception_stepping_tests (after the throw) in mule_main
-                SYMS_ASSERT(!"Hit me!");
+                // SYMS_ASSERT(!"Hit me!");
                 
                 if (op_info > 1){
                   result.dead = syms_true;
@@ -369,13 +369,13 @@ syms_unwind_pe_x64__epilog(SYMS_String8 bin_data, SYMS_PeBinAccel *bin, SYMS_U64
   //- setup parsing context
   SYMS_U64 ip_voff = regs->rip.u64 - bin_base;
   SYMS_U64 sec_number = syms_pe_sec_number_from_voff(bin_data, bin, ip_voff);
-  SYMS_CoffSection coff_sec = syms_pe_coff_section(bin_data, bin, sec_number);
-  void*         inst_base   = syms_pe_ptr_from_sec_number(bin_data, bin, sec_number);
-  SYMS_U64Range inst_range = {0, coff_sec.virt_size};
+  SYMS_CoffSectionHeader *sec = syms_pecoff_sec_hdr_from_n(bin_data, bin->section_array_off, sec_number);
+  void*         inst_base  = syms_pe_ptr_from_sec_number(bin_data, bin, sec_number);
+  SYMS_U64Range inst_range = {0, sec->virt_size};
   
   //- setup parsing variables
   SYMS_B32 keep_parsing = syms_true;
-  SYMS_U64 off = ip_voff - coff_sec.virt_off;
+  SYMS_U64 off = ip_voff - sec->virt_off;
   
   //- parsing loop
   for (;keep_parsing;){
@@ -597,14 +597,14 @@ syms_unwind_pe_x64__voff_is_in_epilog(SYMS_String8 bin_data, SYMS_PeBinAccel *bi
   
   //- setup parsing context
   SYMS_U64 sec_number = syms_pe_sec_number_from_voff(bin_data, bin, voff);
-  SYMS_CoffSection coff_sec = syms_pe_coff_section(bin_data, bin, sec_number);
-  void*         inst_base   = syms_pe_ptr_from_sec_number(bin_data, bin, sec_number);
-  SYMS_U64Range inst_range = {0, coff_sec.virt_size};
+  SYMS_CoffSectionHeader *sec = syms_pecoff_sec_hdr_from_n(bin_data, bin->section_array_off, sec_number);
+  void*         inst_base  = syms_pe_ptr_from_sec_number(bin_data, bin, sec_number);
+  SYMS_U64Range inst_range = {0, sec->virt_size};
   
   //- setup parsing variables
   SYMS_B32 is_epilog = syms_false;
   SYMS_B32 keep_parsing = syms_true;
-  SYMS_U64 off = voff - coff_sec.virt_off;
+  SYMS_U64 off = voff - sec->virt_off;
   
   //- check first instruction
   {
diff --git a/Engine/Source/Runtime/SymsLib/syms/overrides/syms_crt_overrides.c b/Engine/Source/Runtime/SymsLib/syms/overrides/syms_crt_overrides.c
--- a/Engine/Source/Runtime/SymsLib/syms/overrides/syms_crt_overrides.c
+++ b/Engine/Source/Runtime/SymsLib/syms/overrides/syms_crt_overrides.c
@@ -3,10 +3,16 @@
 #ifndef SYMS_CRT_OVERRIDES_C
 #define SYMS_CRT_OVERRIDES_C
 
-// NOTE(allen): use default arena
+////////////////////////////////
+//~ allen: use default helpers
+#include "syms_default_helpers.c"
+
+////////////////////////////////
+//~ allen: use default arena
 #include "syms_default_arena.c"
 
-// NOTE(allen): use default scratch
+////////////////////////////////
+//~ allen: use default scrtach
 #include "syms_default_scratch.c"
 
 #endif // SYMS_CRT_OVERRIDES_C
diff --git a/Engine/Source/Runtime/SymsLib/syms/overrides/syms_crt_overrides.h b/Engine/Source/Runtime/SymsLib/syms/overrides/syms_crt_overrides.h
--- a/Engine/Source/Runtime/SymsLib/syms/overrides/syms_crt_overrides.h
+++ b/Engine/Source/Runtime/SymsLib/syms/overrides/syms_crt_overrides.h
@@ -57,7 +57,8 @@ typedef uint64_t SYMS_U64;
 #include <assert.h>
 #define SYMS_ASSERT_BREAK(m) assert(!(#m))
 
-// NOTE(allen): memory helpers
+////////////////////////////////
+//~ allen: memory helpers
 #include <stdlib.h>
 #include <string.h>
 
@@ -65,8 +66,10 @@ typedef uint64_t SYMS_U64;
 #define syms_memset memset
 #define syms_memcmp memcmp
 #define syms_strlen strlen
+#define syms_memisnull(ptr,sz) (memchr(ptr,0,sz)!=NULL)
 
-// NOTE(allen): memory
+////////////////////////////////
+//~ allen: memory
 #define syms_mem_reserve(s)    malloc(s)
 #define syms_mem_commit(p,s)   (1)
 #define syms_mem_decommit(p,s) ((void)0)
@@ -80,10 +83,12 @@ typedef uint64_t SYMS_U64;
 #define SYMS_ARENA_RESERVE_SIZE (1 << 20)
 #define SYMS_ARENA_COMMIT_SIZE  (1 << 20)
 
-// NOTE(allen): use default arena
+////////////////////////////////
+//~ allen: use default arena
 #include "syms_default_arena.h"
 
-// NOTE(allen): use default scrtach
+////////////////////////////////
+//~ allen: use default scrtach
 #include "syms_default_scratch.h"
 
 #endif // SYMS_CRT_DEFAULT_OVERRIDES_H
diff --git a/Engine/Source/Runtime/SymsLib/syms/syms_inc.c b/Engine/Source/Runtime/SymsLib/syms/syms_inc.c
--- a/Engine/Source/Runtime/SymsLib/syms/syms_inc.c
+++ b/Engine/Source/Runtime/SymsLib/syms/syms_inc.c
@@ -14,7 +14,9 @@
 #include "syms/core/syms_eval.c"
 
 // "windows related" formats
+#include "syms/core/pe/syms_coff.c"
 #include "syms/core/pe/syms_pe.c"
+#include "syms/core/pdb/syms_cv.c"
 #include "syms/core/pdb/syms_pdb.c"
 
 // "linux related" foramts
@@ -25,7 +27,10 @@
 #include "syms/core/mach/syms_mach.c"
 
 // "windows related" parsers
+#include "syms/core/pe/syms_pecoff_helpers.c"
 #include "syms/core/pe/syms_pe_parser.c"
+#include "syms/core/pdb/syms_msf_parser.c"
+#include "syms/core/pdb/syms_cv_helpers.c"
 #include "syms/core/pdb/syms_pdb_parser.c"
 
 // "mach-o related" parsers
@@ -42,6 +47,7 @@
 #include "syms/core/syms_parser_invariants.c"
 #include "syms/core/data_structures/syms_data_structures.c"
 #include "syms/core/group/syms_type_graph.c"
+#include "syms/core/group/syms_functions.c"
 #include "syms/core/group/syms_group.c"
 #include "syms/core/file_inf/syms_file_inf.c"
 
diff --git a/Engine/Source/Runtime/SymsLib/syms/syms_inc.h b/Engine/Source/Runtime/SymsLib/syms/syms_inc.h
--- a/Engine/Source/Runtime/SymsLib/syms/syms_inc.h
+++ b/Engine/Source/Runtime/SymsLib/syms/syms_inc.h
@@ -23,7 +23,10 @@
 #include "syms/core/syms_eval.h"
 
 // "windows related" formats
+#include "syms/core/pe/syms_coff.h"
 #include "syms/core/pe/syms_pe.h"
+#include "syms/core/pdb/syms_cv.h"
+#include "syms/core/pdb/syms_msf.h"
 #include "syms/core/pdb/syms_pdb.h"
 
 // "linux related" formats
@@ -34,7 +37,10 @@
 #include "syms/core/mach/syms_mach.h"
 
 // "windows related" parsers
+#include "syms/core/pe/syms_pecoff_helpers.h"
 #include "syms/core/pe/syms_pe_parser.h"
+#include "syms/core/pdb/syms_msf_parser.h"
+#include "syms/core/pdb/syms_cv_helpers.h"
 #include "syms/core/pdb/syms_pdb_parser.h"
 
 // "mach-o related" parsers
@@ -51,6 +57,7 @@
 #include "syms/core/syms_parser_invariants.h"
 #include "syms/core/data_structures/syms_data_structures.h"
 #include "syms/core/group/syms_type_graph.h"
+#include "syms/core/group/syms_functions.h"
 #include "syms/core/group/syms_group.h"
 #include "syms/core/file_inf/syms_file_inf.h"
 
