Provide universal binaries for SymsLib on Mac

diff --git a/Engine/Source/Runtime/SymsLib/BuildForMac.sh b/Engine/Source/Runtime/SymsLib/BuildForMac.sh
--- a/Engine/Source/Runtime/SymsLib/BuildForMac.sh
+++ b/Engine/Source/Runtime/SymsLib/BuildForMac.sh
@@ -49,3 +49,7 @@ BuildLibSyms x86_64 Release ""
 BuildLibSyms x86_64 Debug "d"
 BuildLibSyms arm64 Release ""
 BuildLibSyms arm64 Debug "d"
+
+lipo -create ${SYMS_DIR}/lib/Mac/{x86_64,arm64}/libsymsd.a -output ${SYMS_DIR}/lib/Mac/libsymsd.a
+lipo -create ${SYMS_DIR}/lib/Mac/{x86_64,arm64}/libsyms.a  -output ${SYMS_DIR}/lib/Mac/libsyms.a
+rm -rf ${SYMS_DIR}/lib/Mac/{x86_64,arm64}
diff --git a/Engine/Source/Runtime/SymsLib/SymsLib.Build.cs b/Engine/Source/Runtime/SymsLib/SymsLib.Build.cs
--- a/Engine/Source/Runtime/SymsLib/SymsLib.Build.cs
+++ b/Engine/Source/Runtime/SymsLib/SymsLib.Build.cs
@@ -32,7 +32,7 @@ public class SymsLib : ModuleRules
 		{
 			bool bUseDebug = Target.Configuration == UnrealTargetConfiguration.Debug;
 			string LibName = bUseDebug ? "libsymsd.a" : "libsyms.a";
-			string LibPath = Path.Combine(LibPathBase, "Mac", Target.Architecture.bIsX64 ? "x86_64" : "arm64");
+			string LibPath = Path.Combine(LibPathBase, "Mac");
 
 			PublicAdditionalLibraries.Add(Path.Combine(LibPath, LibName));
 		}
