Update SymsLib code with latest changes.
Adds support for zlib compressed debug sections to elf/dwarf parser.
Some minor fixes to Mach-O & dwarf parsing.
Fixes crash in SymslibResolver when there are no compilation units.

diff --git a/Engine/Source/Developer/TraceServices/Private/Model/SymslibResolver.cpp b/Engine/Source/Developer/TraceServices/Private/Model/SymslibResolver.cpp
--- a/Engine/Source/Developer/TraceServices/Private/Model/SymslibResolver.cpp
+++ b/Engine/Source/Developer/TraceServices/Private/Model/SymslibResolver.cpp
@@ -1121,6 +1121,10 @@ EModuleStatus FSymslibResolver::LoadModule(FModuleEntry* Entry, FStringView Sear
 	}
 
 	SYMS_Arena* Arena = Instance->Arenas[0];
+	if (!Arena)
+	{
+		Arena = Instance->Arenas[0] = syms_arena_alloc();
+	}
 
 	// store stripped format symbols
 	{
diff --git a/Engine/Source/Runtime/SymsLib/BuildForLinux.sh b/Engine/Source/Runtime/SymsLib/BuildForLinux.sh
--- a/Engine/Source/Runtime/SymsLib/BuildForLinux.sh
+++ b/Engine/Source/Runtime/SymsLib/BuildForLinux.sh
@@ -2,14 +2,12 @@
 
 set -eu
 
-ARCH=x86_64-unknown-linux-gnu
 SCRIPT_DIR=$(cd "$(dirname "$BASH_SOURCE")" ; pwd)
 UE_THIRD_PARTY_DIR=$(cd "${SCRIPT_DIR}/../../../.."; pwd)
 SYMS_DIR=$(cd "${SCRIPT_DIR}"; pwd)
 
-CFLAGS="-ffunction-sections -fdata-sections"
+CFLAGS="-ffunction-sections -fdata-sections -I$SCRIPT_DIR/../../ThirdParty/zlib/1.3"
 CXXFLAGS="-ffunction-sections -fdata-sections -I$UE_THIRD_PARTY_DIR/Unix/LibCxx/include -I$UE_THIRD_PARTY_DIR/Unix/LibCxx/include/c++/v1"
-LIBS="$UE_THIRD_PARTY_DIR/Unix/LibCxx/lib/Unix/${ARCH}/libc++.a $UE_THIRD_PARTY_DIR/Unix/LibCxx/lib/Unix/${ARCH}/libc++abi.a -lm -lc -lgcc_s -lgcc -lpthread"
 
 # Get num of cores
 export CORES=$(getconf _NPROCESSORS_ONLN)
@@ -22,28 +20,9 @@ BuildLibSyms()
     local EXT=$3
     local BUILD_DIR=/tmp/Build-Syms-${ARCH}-${FLAVOR}
 
-    echo "Building ${ARCH} ${FLAVOR}"
-    rm -rf ${BUILD_DIR}
-    mkdir -p ${BUILD_DIR}/Build
-
-    pushd ${BUILD_DIR}/Build
-
-    set -x
-    cmake \
-      -DCMAKE_TOOLCHAIN_FILE="/tmp/__cmake_toolchain.cmake" \
-      -DCMAKE_MAKE_PROGRAM=$(which make) \
-      -DCMAKE_BUILD_TYPE=${FLAVOR} \
-      -DCMAKE_POSITION_INDEPENDENT_CODE=ON \
-      ${SYMS_DIR}/
-    set +x
+    LIBS="$UE_THIRD_PARTY_DIR/Unix/LibCxx/lib/Unix/${ARCH}/libc++.a $UE_THIRD_PARTY_DIR/Unix/LibCxx/lib/Unix/${ARCH}/libc++abi.a -lm -lc -lgcc_s -lgcc -lpthread"
 
-    VERBOSE=1 make -j ${CORES}
-
-    mkdir -p ${SYMS_DIR}/lib/Unix/${ARCH}
-    cp -v ${BUILD_DIR}/Build/libsyms.a ${SYMS_DIR}/lib/Unix/${ARCH}/libsyms${EXT}_fPIC.a
-
-    popd
-}
+    echo "Building ${ARCH} ${FLAVOR}"
 
 ( cat <<_EOF_
   ## autogenerated by ${BASH_SOURCE} script
@@ -85,5 +64,29 @@ BuildLibSyms()
 _EOF_
 ) > /tmp/__cmake_toolchain.cmake
 
-BuildLibSyms ${ARCH} Release ""
-BuildLibSyms ${ARCH} Debug "d"
+    rm -rf ${BUILD_DIR}
+    mkdir -p ${BUILD_DIR}/Build
+
+    pushd ${BUILD_DIR}/Build
+
+    set -x
+    cmake \
+      -DCMAKE_TOOLCHAIN_FILE="/tmp/__cmake_toolchain.cmake" \
+      -DCMAKE_MAKE_PROGRAM=$(which make) \
+      -DCMAKE_BUILD_TYPE=${FLAVOR} \
+      -DCMAKE_POSITION_INDEPENDENT_CODE=ON \
+      ${SYMS_DIR}/
+    set +x
+
+    VERBOSE=1 make -j ${CORES}
+
+    mkdir -p ${SYMS_DIR}/lib/Unix/${ARCH}
+    cp -v ${BUILD_DIR}/Build/libsyms.a ${SYMS_DIR}/lib/Unix/${ARCH}/libsyms${EXT}_fPIC.a
+
+    popd
+}
+
+BuildLibSyms x86_64-unknown-linux-gnu      Release ""
+BuildLibSyms x86_64-unknown-linux-gnu      Debug   "d"
+BuildLibSyms aarch64-unknown-linux-gnueabi Release ""
+BuildLibSyms aarch64-unknown-linux-gnueabi Debug   "d"
diff --git a/Engine/Source/Runtime/SymsLib/BuildForMac.sh b/Engine/Source/Runtime/SymsLib/BuildForMac.sh
--- a/Engine/Source/Runtime/SymsLib/BuildForMac.sh
+++ b/Engine/Source/Runtime/SymsLib/BuildForMac.sh
@@ -6,7 +6,7 @@ SCRIPT_DIR=$(cd "$(dirname "$BASH_SOURCE")" ; pwd)
 UE_THIRD_PARTY_DIR=$(cd "${SCRIPT_DIR}/../../../.."; pwd)
 SYMS_DIR=$(cd "${SCRIPT_DIR}"; pwd)
 
-CFLAGS="-ffunction-sections -fdata-sections"
+CFLAGS="-ffunction-sections -fdata-sections -I$SCRIPT_DIR/../../ThirdParty/zlib/1.3"
 CXXFLAGS="-ffunction-sections -fdata-sections"
 LIBS="-lm -lc -lpthread"
 
diff --git a/Engine/Source/Runtime/SymsLib/SymsLib.Build.cs b/Engine/Source/Runtime/SymsLib/SymsLib.Build.cs
--- a/Engine/Source/Runtime/SymsLib/SymsLib.Build.cs
+++ b/Engine/Source/Runtime/SymsLib/SymsLib.Build.cs
@@ -37,5 +37,7 @@ public class SymsLib : ModuleRules
 
 			PublicAdditionalLibraries.Add(Path.Combine(LibPath, LibName));
 		}
+
+        AddEngineThirdPartyPrivateStaticDependencies(Target, "zlib");
 	}
 }
diff --git a/Engine/Source/Runtime/SymsLib/syms/core/dwarf/syms_dwarf_parser.c b/Engine/Source/Runtime/SymsLib/syms/core/dwarf/syms_dwarf_parser.c
--- a/Engine/Source/Runtime/SymsLib/syms/core/dwarf/syms_dwarf_parser.c
+++ b/Engine/Source/Runtime/SymsLib/syms/core/dwarf/syms_dwarf_parser.c
@@ -749,6 +749,25 @@ syms_dw_elf_bin_accel_is_dbg(SYMS_ElfBinAccel *bin_accel)
   return result;
 }
 
+SYMS_API SYMS_B32
+syms_dw_mach_bin_accel_is_dbg(SYMS_MachBinAccel *bin_accel)
+{
+  SYMS_ArenaTemp scratch = syms_get_scratch(0,0);
+  SYMS_B32 is_dbg = syms_false;
+  SYMS_SecInfoArray sec_array = syms_mach_sec_info_array_from_bin(scratch.arena, syms_str8(0,0), bin_accel);
+  for(SYMS_U64 i = 0; i < sec_array.count; i += 1)
+  {
+    if(syms_string_match(sec_array.sec_info[i].name, syms_str8_lit("__debug_info"), 0) ||
+       syms_string_match(sec_array.sec_info[i].name, syms_str8_lit("__debug_line"), 0))
+    {
+      is_dbg = syms_true;
+      break;
+    }
+  }
+  syms_release_scratch(scratch);
+  return is_dbg;
+}
+
 SYMS_API SYMS_DwDbgAccel *
 syms_dw_dbg_accel_from_sec_info_array(SYMS_Arena *arena, SYMS_String8 data, SYMS_U64 vbase, SYMS_Arch arch, SYMS_SecInfoArray sections)
 {
@@ -793,8 +812,7 @@ syms_dw_dbg_accel_from_sec_info_array(SYMS_Arena *arena, SYMS_String8 data, SYMS
     }
   }
   
-  //- rjf: build section map + figure out mode
-  SYMS_DwMode mode = SYMS_DwMode_32Bit;
+  //- rjf: build section map
   SYMS_DwSection *section_map = 0;
   SYMS_B32 is_dwo = syms_false;
   if(has_necessary_sections)
@@ -815,10 +833,6 @@ syms_dw_dbg_accel_from_sec_info_array(SYMS_Arena *arena, SYMS_String8 data, SYMS
         if(match_main || match_dwo || match_mach)
         {
           section_map[section_kind].range = sections.sec_info[section_idx].frange;
-          if(section_map[section_kind].range.max - section_map[section_kind].range.min > SYMS_U32_MAX)
-          {
-            mode = SYMS_DwMode_64Bit;
-          }
           if(match_dwo && section_kind == SYMS_DwSectionKind_Info)
           {
             is_dwo = syms_true;
@@ -843,12 +857,14 @@ syms_dw_dbg_accel_from_sec_info_array(SYMS_Arena *arena, SYMS_String8 data, SYMS
   //- rjf: find .debug_info and .debug_aranges for initial quick parse
   void *info_base = 0;
   SYMS_U64Range info_rng = {0};
+  SYMS_DwMode aranges_mode = SYMS_DwMode_Null;
   void *aranges_base = 0;
   SYMS_U64Range aranges_rng = {0};
   if(has_necessary_sections)
   {
     info_base    = data.str + section_map[SYMS_DwSectionKind_Info].range.min;
     info_rng     = syms_make_u64_range(0, syms_u64_range_size(section_map[SYMS_DwSectionKind_Info].range));
+    aranges_mode = syms_dw_mode_from_sec(section_map, SYMS_DwSectionKind_ARanges);
     aranges_base = data.str + section_map[SYMS_DwSectionKind_ARanges].range.min;
     aranges_rng  = syms_make_u64_range(0, syms_u64_range_size(section_map[SYMS_DwSectionKind_ARanges].range));
   }
@@ -893,7 +909,7 @@ syms_dw_dbg_accel_from_sec_info_array(SYMS_Arena *arena, SYMS_String8 data, SYMS
         SYMS_U8 segment_size = 0;
         offset += syms_dw_based_range_read_length(aranges_base, aranges_rng, offset, &unit_length);
         offset += syms_based_range_read_struct(aranges_base, aranges_rng, offset, &version);
-        offset += syms_based_range_read(aranges_base, aranges_rng, offset, syms_dw_offset_size_from_mode(mode), &debug_info_offset);
+        offset += syms_based_range_read(aranges_base, aranges_rng, offset, syms_dw_offset_size_from_mode(aranges_mode), &debug_info_offset);
         offset += syms_based_range_read_struct(aranges_base, aranges_rng, offset, &address_size);
         offset += syms_based_range_read_struct(aranges_base, aranges_rng, offset, &segment_size);
         
@@ -986,7 +1002,6 @@ syms_dw_dbg_accel_from_sec_info_array(SYMS_Arena *arena, SYMS_String8 data, SYMS
     dbg = syms_push_array_zero(arena, SYMS_DwDbgAccel, 1);
     dbg->format = SYMS_FileFormat_DWARF;
     dbg->arch = arch;
-    dbg->mode = mode;
     dbg->vbase = vbase;
     dbg->sections = syms_dw_copy_sec_info_array(arena, sections);
     dbg->text_section_idx = text_section_idx;
@@ -1102,6 +1117,21 @@ syms_dw_uid_from_foff(SYMS_DwDbgAccel *dbg, SYMS_U64 foff)
 
 //- rjf: important DWARF section base/range accessors
 
+SYMS_API SYMS_DwMode
+syms_dw_mode_from_sec(SYMS_DwSection *section_map, SYMS_DwSectionKind kind)
+{
+  SYMS_U64Range range = section_map[kind].range;
+  SYMS_U64 size = range.max - range.min;
+  if(size > 0xffffffff)
+  {
+    return SYMS_DwMode_64Bit;
+  }
+  else
+  {
+    return SYMS_DwMode_32Bit;
+  }
+}
+
 SYMS_API SYMS_B32
 syms_dw_sec_is_present(SYMS_DwDbgAccel *dbg, SYMS_DwSectionKind kind)
 {
@@ -1366,11 +1396,12 @@ syms_dw_v4_pub_strings_table_from_section_kind(SYMS_Arena *arena, SYMS_String8 d
   names_table.size = 16384;
   names_table.buckets = syms_push_array_zero(arena, SYMS_DwPubStringsBucket*, names_table.size);
   
+  SYMS_DwMode mode = syms_dw_mode_from_sec(dbg->section_map, section_kind);
   void *base = syms_dw_sec_base_from_dbg(data, dbg, section_kind);
   SYMS_U64Range rng = syms_dw_sec_range_from_dbg(dbg, section_kind);
   SYMS_U64 read_off = 0;
   
-  SYMS_U64 off_size = syms_dw_offset_size_from_mode(dbg->mode);
+  SYMS_U64 off_size = syms_dw_offset_size_from_mode(mode);
   
   SYMS_U64 table_length = 0;
   SYMS_U16 unit_version = 0;
@@ -1462,12 +1493,13 @@ syms_dw_v5_offset_from_offs_section_base_index(SYMS_String8 data, SYMS_DwDbgAcce
 {
   SYMS_U64 result = 0;
   
+  SYMS_DwMode mode = syms_dw_mode_from_sec(dbg->section_map, section);
   void *sec_base = syms_dw_sec_base_from_dbg(data, dbg, section);
   SYMS_U64Range rng = syms_dw_sec_range_from_dbg(dbg, section);
   SYMS_U64 read_off = base;
   
   //- rjf: get the length of each entry
-  SYMS_U64 entry_len = dbg->mode == SYMS_DwMode_64Bit ? 8 : 4;
+  SYMS_U64 entry_len = mode == SYMS_DwMode_64Bit ? 8 : 4;
   
   //- rjf: parse the unit's length (not including the length itself)
   SYMS_U64 unit_length = 0;
@@ -1562,12 +1594,13 @@ syms_dw_v5_sec_offset_from_rnglist_or_loclist_section_base_index(SYMS_String8 da
   //
   
   SYMS_U64 result = 0;
+  SYMS_DwMode mode = syms_dw_mode_from_sec(dbg->section_map, section_kind);
   void *sec_base = syms_dw_sec_base_from_dbg(data, dbg, section_kind);
   SYMS_U64Range rng = syms_dw_sec_range_from_dbg(dbg, section_kind);
   SYMS_U64 read_off = base;
   
   //- rjf: get the length of each entry
-  SYMS_U64 entry_len = dbg->mode == SYMS_DwMode_64Bit ? 8 : 4;
+  SYMS_U64 entry_len = mode == SYMS_DwMode_64Bit ? 8 : 4;
   
   //- rjf: parse the unit's length (not including the length itself)
   SYMS_U64 unit_length = 0;
@@ -2096,8 +2129,10 @@ SYMS_API SYMS_DwAttribListParseResult
 syms_dw_parse_attrib_list_from_info_abbrev_offsets(SYMS_Arena *arena, SYMS_String8 data, SYMS_DwDbgAccel *dbg, SYMS_DwLanguage lang, SYMS_DwVersion ver, SYMS_U64 address_size, SYMS_U64 info_off, SYMS_U64 abbrev_off)
 {
   //- rjf: set up prereqs
+  SYMS_DwMode info_mode = syms_dw_mode_from_sec(dbg->section_map, SYMS_DwSectionKind_Info);
   void *info_base = syms_dw_sec_base_from_dbg(data, dbg, SYMS_DwSectionKind_Info);
   SYMS_U64Range info_range = syms_dw_sec_range_from_dbg(dbg, SYMS_DwSectionKind_Info);
+  SYMS_DwMode abbrev_mode = syms_dw_mode_from_sec(dbg->section_map, SYMS_DwSectionKind_Abbrev);
   void *abbrev_base = syms_dw_sec_base_from_dbg(data, dbg, SYMS_DwSectionKind_Abbrev);
   SYMS_U64Range abbrev_range = syms_dw_sec_range_from_dbg(dbg, SYMS_DwSectionKind_Abbrev);
   
@@ -2138,7 +2173,7 @@ syms_dw_parse_attrib_list_from_info_abbrev_offsets(SYMS_Arena *arena, SYMS_Strin
         info_read_off += syms_based_range_read_uleb128(info_base, info_range, info_read_off, &override_form_kind);
         form_kind = (SYMS_DwFormKind)override_form_kind;
       }
-      SYMS_U64 bytes_read = syms_dw_based_range_read_attrib_form_value(info_base, info_range, info_read_off, dbg->mode, address_size,
+      SYMS_U64 bytes_read = syms_dw_based_range_read_attrib_form_value(info_base, info_range, info_read_off, info_mode, address_size,
                                                                        form_kind, abbrev.const_value, &form_value);
       info_read_off += bytes_read;
     }
@@ -2608,10 +2643,11 @@ syms_dw_comp_root_from_range(SYMS_Arena *arena, SYMS_String8 data, SYMS_DwDbgAcc
   
   //- rjf: build+fill directory and file tables
   {
+    SYMS_DwMode line_mode = syms_dw_mode_from_sec(dbg->section_map, SYMS_DwSectionKind_Line);
     void *line_base = syms_dw_sec_base_from_dbg(data, dbg, SYMS_DwSectionKind_Line);
     SYMS_U64Range line_rng = syms_dw_sec_range_from_dbg(dbg, SYMS_DwSectionKind_Line);
     SYMS_DwLineVMHeader vm_header;
-    SYMS_U64 read_size = syms_dw_read_line_vm_header(arena, line_base, line_rng, unit.line_off, data, dbg, &unit, &vm_header);
+    SYMS_U64 read_size = syms_dw_read_line_vm_header(arena, line_base, line_rng, unit.line_off, data, line_mode, dbg, &unit, &vm_header);
     if (read_size > 0) {
       unit.dir_table = vm_header.dir_table;
       unit.file_table = vm_header.file_table;
@@ -5095,13 +5131,14 @@ syms_dw_parsed_line_table_from_comp_root(SYMS_Arena *arena, SYMS_String8 data, S
 {
   SYMS_LogOpen(SYMS_LogFeature_LineTable, root->index + 1, dwarf_line_vm);
   
+  SYMS_DwMode mode = syms_dw_mode_from_sec(dbg->section_map, SYMS_DwSectionKind_Line);
   void *base = syms_dw_sec_base_from_dbg(data, dbg, SYMS_DwSectionKind_Line);
   SYMS_U64Range line_info_range = syms_dw_sec_range_from_dbg(dbg, SYMS_DwSectionKind_Line);
   SYMS_U64 read_off_start = root->line_off - line_info_range.min;
   SYMS_U64 read_off = read_off_start;
   
   SYMS_DwLineVMHeader vm_header;
-  read_off += syms_dw_read_line_vm_header(arena, base, line_info_range, read_off, data, dbg, root, &vm_header);
+  read_off += syms_dw_read_line_vm_header(arena, base, line_info_range, read_off, data, mode, dbg, root, &vm_header);
   
   //- rjf: prep state for VM
   SYMS_DwLineVMState vm_state;
@@ -5389,7 +5426,7 @@ syms_dw_line_parse_from_uid(SYMS_Arena *arena, SYMS_String8 data, SYMS_DwDbgAcce
 }
 
 SYMS_API SYMS_U64
-syms_dw_read_line_file(void *line_base, SYMS_U64Range line_rng, SYMS_U64 line_off, SYMS_DwDbgAccel *dbg, SYMS_String8 data, SYMS_DwCompRoot *unit, SYMS_U8 address_size, SYMS_U64 format_count, SYMS_U64Range *formats, SYMS_DwLineFile *line_file_out)
+syms_dw_read_line_file(void *line_base, SYMS_U64Range line_rng, SYMS_U64 line_off, SYMS_DwMode mode, SYMS_DwDbgAccel *dbg, SYMS_String8 data, SYMS_DwCompRoot *unit, SYMS_U8 address_size, SYMS_U64 format_count, SYMS_U64Range *formats, SYMS_DwLineFile *line_file_out)
 {
   syms_memzero_struct(line_file_out);
   
@@ -5401,7 +5438,7 @@ syms_dw_read_line_file(void *line_base, SYMS_U64Range line_rng, SYMS_U64 line_of
     SYMS_DwFormKind form_kind = (SYMS_DwFormKind)formats[format_idx].max;
     SYMS_DwAttribValue form_value;
     syms_memzero_struct(&form_value);
-    line_off += syms_dw_based_range_read_attrib_form_value(line_base, line_rng, line_off, dbg->mode, address_size, form_kind, 0, &form_value);
+    line_off += syms_dw_based_range_read_attrib_form_value(line_base, line_rng, line_off, mode, address_size, form_kind, 0, &form_value);
     switch (lnct) 
     {
       case SYMS_DwLNCT_PATH: 
@@ -5459,7 +5496,7 @@ syms_dw_read_line_file(void *line_base, SYMS_U64Range line_rng, SYMS_U64 line_of
 }
 
 SYMS_API SYMS_U64
-syms_dw_read_line_vm_header(SYMS_Arena *arena, void *line_base, SYMS_U64Range line_rng, SYMS_U64 line_off, SYMS_String8 data, SYMS_DwDbgAccel *dbg, SYMS_DwCompRoot *unit, SYMS_DwLineVMHeader *header_out)
+syms_dw_read_line_vm_header(SYMS_Arena *arena, void *line_base, SYMS_U64Range line_rng, SYMS_U64 line_off, SYMS_String8 data, SYMS_DwMode mode, SYMS_DwDbgAccel *dbg, SYMS_DwCompRoot *unit, SYMS_DwLineVMHeader *header_out)
 {
   SYMS_U64 line_off_start = line_off;
   
@@ -5482,7 +5519,7 @@ syms_dw_read_line_vm_header(SYMS_Arena *arena, void *line_base, SYMS_U64Range li
   }
   
   {
-    SYMS_U64 off_size = syms_dw_offset_size_from_mode(dbg->mode);
+    SYMS_U64 off_size = syms_dw_offset_size_from_mode(mode);
     line_off += syms_based_range_read(line_base, line_rng, line_off, off_size, &header_out->header_length);
   }
   
@@ -5549,6 +5586,7 @@ syms_dw_read_line_vm_header(SYMS_Arena *arena, void *line_base, SYMS_U64Range li
       line_off += syms_dw_read_line_file(line_base,
                                          line_rng,
                                          line_off,
+                                         mode,
                                          dbg,
                                          data,
                                          unit,
@@ -5582,6 +5620,7 @@ syms_dw_read_line_vm_header(SYMS_Arena *arena, void *line_base, SYMS_U64Range li
       line_off += syms_dw_read_line_file(line_base,
                                          line_rng,
                                          line_off,
+                                         mode,
                                          dbg,
                                          data,
                                          unit,
diff --git a/Engine/Source/Runtime/SymsLib/syms/core/dwarf/syms_dwarf_parser.h b/Engine/Source/Runtime/SymsLib/syms/core/dwarf/syms_dwarf_parser.h
--- a/Engine/Source/Runtime/SymsLib/syms/core/dwarf/syms_dwarf_parser.h
+++ b/Engine/Source/Runtime/SymsLib/syms/core/dwarf/syms_dwarf_parser.h
@@ -95,7 +95,6 @@ struct SYMS_DwDbgAccel
   SYMS_FileFormat format;
   SYMS_Arch arch;
   SYMS_U64 vbase;
-  SYMS_DwMode mode;
   SYMS_SecInfoArray sections;
   SYMS_U64 text_section_idx;
   SYMS_U64Range acceptable_vrange;
@@ -598,6 +597,7 @@ SYMS_API SYMS_U64 syms_dw_based_range_read_attrib_form_value(void *base, SYMS_U6
 
 //- rjf: debug info and interaction with bin accels
 SYMS_API SYMS_B32         syms_dw_elf_bin_accel_is_dbg(SYMS_ElfBinAccel *bin_accel);
+SYMS_API SYMS_B32         syms_dw_mach_bin_accel_is_dbg(SYMS_MachBinAccel *bin_accel);
 SYMS_API SYMS_DwDbgAccel *syms_dw_dbg_accel_from_sec_info_array(SYMS_Arena *arena, SYMS_String8 data, SYMS_U64 vbase, SYMS_Arch arch, SYMS_SecInfoArray sections);
 SYMS_API SYMS_DwDbgAccel *syms_dw_dbg_accel_from_elf_bin(SYMS_Arena *arena, SYMS_String8 data, SYMS_ElfBinAccel *bin);
 SYMS_API SYMS_DwDbgAccel *syms_dw_dbg_accel_from_mach_bin(SYMS_Arena *arena, SYMS_String8 data, SYMS_MachBinAccel *bin);
@@ -611,6 +611,8 @@ SYMS_API SYMS_U64          syms_dw_default_vbase_from_dbg(SYMS_DwDbgAccel *dbg);
 SYMS_API SYMS_UnitID syms_dw_uid_from_foff(SYMS_DwDbgAccel *dbg, SYMS_U64 foff);
 
 //- rjf: important DWARF section base/range accessors
+SYMS_API SYMS_DwMode syms_dw_mode_from_sec(SYMS_DwSection *section_map, SYMS_DwSectionKind kind);
+
 SYMS_API SYMS_B32      syms_dw_sec_is_present(SYMS_DwDbgAccel *dbg, SYMS_DwSectionKind kind);
 SYMS_API void*         syms_dw_sec_base_from_dbg(SYMS_String8 data, SYMS_DwDbgAccel *dbg,
                                                  SYMS_DwSectionKind kind);
@@ -832,12 +834,12 @@ SYMS_API SYMS_DwLineTableParseResult syms_dw_parsed_line_table_from_comp_root(SY
 SYMS_API SYMS_LineParseOut syms_dw_line_parse_from_uid(SYMS_Arena *arena, SYMS_String8 data, SYMS_DwDbgAccel *dbg,
                                                        SYMS_DwUnitSetAccel *set, SYMS_UnitID uid);
 SYMS_API SYMS_U64 syms_dw_read_line_file(void *line_base, SYMS_U64Range line_rng, SYMS_U64 line_off,
-                                         SYMS_DwDbgAccel *dbg,  SYMS_String8 data, SYMS_DwCompRoot *unit,
+                                         SYMS_DwMode mode, SYMS_DwDbgAccel *dbg,  SYMS_String8 data, SYMS_DwCompRoot *unit,
                                          SYMS_U8 address_size, 
                                          SYMS_U64 format_count, SYMS_U64Range *formats,
                                          SYMS_DwLineFile *line_file_out);
 SYMS_API SYMS_U64 syms_dw_read_line_vm_header(SYMS_Arena *arena, void *line_base, SYMS_U64Range line_rng,
-                                              SYMS_U64 line_off, SYMS_String8 data,  SYMS_DwDbgAccel *dbg,
+                                              SYMS_U64 line_off, SYMS_String8 data, SYMS_DwMode mode, SYMS_DwDbgAccel *dbg,
                                               SYMS_DwCompRoot *unit, SYMS_DwLineVMHeader *header_out);
 
 //- rjf: name maps
diff --git a/Engine/Source/Runtime/SymsLib/syms/core/elf/syms_elf.h b/Engine/Source/Runtime/SymsLib/syms/core/elf/syms_elf.h
--- a/Engine/Source/Runtime/SymsLib/syms/core/elf/syms_elf.h
+++ b/Engine/Source/Runtime/SymsLib/syms/core/elf/syms_elf.h
@@ -282,6 +282,38 @@ typedef struct
   // SYMS_U8  data[1];
 } SYMS_ElfNote;
 
+////////////////////////////////
+// Extensions
+
+typedef SYMS_U8 SYMS_ElfCompressType;
+enum SYMS_ElfCompressTypeEnum
+{
+  SYMS_ElfCompressType_None = 0,
+  SYMS_ElfCompressType_ZLib = 1,
+  SYMS_ElfCompressType_ZStd = 2,
+  
+  SYMS_ElfCompressType_LoOs = 0x60000000,
+  SYMS_ElfCompressType_HiOs = 0x6fffffff,
+  
+  SYMS_ElfCompressType_LoProc = 0x70000000,
+  SYMS_ElfCompressType_HiProc = 0x7fffffff,
+};
+
+typedef struct SYMS_ElfChdr32
+{
+  SYMS_U32 ch_type;
+  SYMS_U32 ch_size;
+  SYMS_U32 ch_addr_align;
+} SYMS_ElfChdr32;
+
+typedef struct SYMS_ElfChdr64
+{
+  SYMS_U32 ch_type;
+  SYMS_U32 ch_reserved;
+  SYMS_U64 ch_size;
+  SYMS_U64 ch_addr_align;
+} SYMS_ElfChdr64;
+
 SYMS_C_LINKAGE_BEGIN
 
 ////////////////////////////////
diff --git a/Engine/Source/Runtime/SymsLib/syms/core/elf/syms_elf_parser.c b/Engine/Source/Runtime/SymsLib/syms/core/elf/syms_elf_parser.c
--- a/Engine/Source/Runtime/SymsLib/syms/core/elf/syms_elf_parser.c
+++ b/Engine/Source/Runtime/SymsLib/syms/core/elf/syms_elf_parser.c
@@ -3,6 +3,8 @@
 #ifndef SYMS_ELF_PARSER_C
 #define SYMS_ELF_PARSER_C
 
+#include <zlib.h>
+
 ////////////////////////////////
 //~ rjf: Low-Level Header/Section Parsing
 
@@ -213,8 +215,67 @@ syms_elf_section_array_from_img_header(SYMS_Arena *arena, SYMS_String8 file, SYM
     sections[section_idx].virtual_range = syms_make_u64_range(header.sh_addr, header.sh_addr + virt_size);
     sections[section_idx].file_range = syms_make_u64_range(header.sh_offset, header.sh_offset + file_size);
     sections[section_idx].name = name_stabilized;
+
+    if (header.sh_flags & SYMS_ElfSectionFlag_COMPRESSED)
+    {
+      SYMS_ElfChdr64 chdr64 = { 0 };
+      SYMS_U64 chdr_size = 0;
+      if (img.is_32bit)
+      {
+        SYMS_ElfChdr32 chdr32 = { 0 };
+        chdr_size = syms_based_range_read_struct(file_base, sections[section_idx].file_range, 0, &chdr32);
+        if (chdr_size == sizeof(chdr32))
+        {
+          chdr64.ch_type = chdr32.ch_type;
+          chdr64.ch_size = chdr32.ch_size;
+          chdr64.ch_addr_align = chdr32.ch_addr_align;
+        }
+      }
+      else
+      {
+        chdr_size = syms_based_range_read_struct(file_base, sections[section_idx].file_range, 0, &chdr64);
+      }
+
+      if (chdr64.ch_type == SYMS_ElfCompressType_None)
+      {
+        // ok, no compression
+      }
+      else if (chdr64.ch_type == SYMS_ElfCompressType_ZLib)
+      {
+        uLongf uncompressed_size = chdr64.ch_size;
+        uLongf compressed_size = file_size;
+        if (compressed_size >= chdr_size) compressed_size -= chdr_size;
+        
+        SYMS_U8* uncompressed = syms_push_array(arena, SYMS_U8, uncompressed_size); // TODO: chdr64.ch_addr_align
+        Bytef* compressed = (Bytef*)file_base + header.sh_offset + chdr_size;
+        
+        int status = uncompress(uncompressed, &uncompressed_size, compressed, compressed_size);
+        if (status != Z_OK)
+        {
+          // failed to decompress with zlib, this is failure
+          result.v = 0;
+          result.count = 0;
+          return result;
+        }
+
+        // adjust range so when original data pointer is added, we get offsets into newly allocated space
+        sections[section_idx].file_range.min = uncompressed - (SYMS_U8*)file_base;
+        sections[section_idx].file_range.max = uncompressed - (SYMS_U8*)file_base + uncompressed_size;
+      }
+      //else if (chdr64.ch_type == ELF_CompressType_ZStd)
+      //{
+      //  TODO
+      //}
+      else
+      {
+        // unknown compression type, cannot continue parsing
+        result.v = 0;
+        result.count = 0;
+        return result;
+      }
+    }
   }
-  
+
   //- rjf: fill result
   result.v = sections;
   result.count = section_count;
diff --git a/Engine/Source/Runtime/SymsLib/syms/core/syms_parser.c b/Engine/Source/Runtime/SymsLib/syms/core/syms_parser.c
--- a/Engine/Source/Runtime/SymsLib/syms/core/syms_parser.c
+++ b/Engine/Source/Runtime/SymsLib/syms/core/syms_parser.c
@@ -361,6 +361,10 @@ syms_bin_is_dbg(SYMS_BinAccel *bin){
     {
       result = syms_dw_elf_bin_accel_is_dbg(&bin->elf_accel);
     }break;
+    case SYMS_FileFormat_MACH:
+    {
+      result = syms_dw_mach_bin_accel_is_dbg(&bin->mach_accel);
+    }break;
   }
   SYMS_ASSERT_PARANOID(syms_parser_api_invariants());
   SYMS_ProfEnd();
