Add range checks in elf header reading in SymsLib
Allow it to decompress large zlib dwarf sections - use zlib inflate instead of uncompress function.

diff --git a/Engine/Source/Runtime/SymsLib/syms/core/elf/syms_elf_parser.c b/Engine/Source/Runtime/SymsLib/syms/core/elf/syms_elf_parser.c
--- a/Engine/Source/Runtime/SymsLib/syms/core/elf/syms_elf_parser.c
+++ b/Engine/Source/Runtime/SymsLib/syms/core/elf/syms_elf_parser.c
@@ -218,39 +218,94 @@ syms_elf_section_array_from_img_header(SYMS_Arena *arena, SYMS_String8 file, SYM
 
     if (header.sh_flags & SYMS_ElfSectionFlag_COMPRESSED)
     {
-      SYMS_ElfChdr64 chdr64 = { 0 };
+      SYMS_U64Range chdr_range = sections[section_idx].file_range;
+      chdr_range.min = SYMS_ClampTop(chdr_range.min, syms_u64_range_size(file_range));
+      chdr_range.max = SYMS_ClampTop(chdr_range.max, syms_u64_range_size(file_range));
+      
+      SYMS_ElfChdr64 chdr = { 0 };
       SYMS_U64 chdr_size = 0;
       if (img.is_32bit)
       {
         SYMS_ElfChdr32 chdr32 = { 0 };
-        chdr_size = syms_based_range_read_struct(file_base, sections[section_idx].file_range, 0, &chdr32);
+        chdr_size = syms_based_range_read_struct(file_base, chdr_range, 0, &chdr32);
         if (chdr_size == sizeof(chdr32))
         {
-          chdr64.ch_type = chdr32.ch_type;
-          chdr64.ch_size = chdr32.ch_size;
-          chdr64.ch_addr_align = chdr32.ch_addr_align;
+          chdr.ch_type = chdr32.ch_type;
+          chdr.ch_size = chdr32.ch_size;
+          chdr.ch_addr_align = chdr32.ch_addr_align;
         }
       }
       else
       {
-        chdr_size = syms_based_range_read_struct(file_base, sections[section_idx].file_range, 0, &chdr64);
+        SYMS_ElfChdr64 chdr64 = { 0 };
+        chdr_size = syms_based_range_read_struct(file_base, chdr_range, 0, &chdr64);
+        if (chdr_size == sizeof(chdr64))
+        {
+          chdr.ch_type = chdr64.ch_type;
+          chdr.ch_size = chdr64.ch_size;
+          chdr.ch_addr_align = chdr64.ch_addr_align;
+        }
       }
 
-      if (chdr64.ch_type == SYMS_ElfCompressType_None)
+      if (chdr.ch_type == SYMS_ElfCompressType_None)
       {
         // ok, no compression
       }
-      else if (chdr64.ch_type == SYMS_ElfCompressType_ZLib)
+      else if (chdr.ch_type == SYMS_ElfCompressType_ZLib)
       {
-        uLongf uncompressed_size = chdr64.ch_size;
-        uLongf compressed_size = file_size;
+        SYMS_U64 uncompressed_size = chdr.ch_size;
+        SYMS_U64 compressed_size = file_size;
+        if (compressed_size > syms_u64_range_size(file_range) - header.sh_offset) compressed_size = 0;
         if (compressed_size >= chdr_size) compressed_size -= chdr_size;
         
-        SYMS_U8* uncompressed = syms_push_array(arena, SYMS_U8, uncompressed_size); // TODO: chdr64.ch_addr_align
-        Bytef* compressed = (Bytef*)file_base + header.sh_offset + chdr_size;
-        
-        int status = uncompress(uncompressed, &uncompressed_size, compressed, compressed_size);
-        if (status != Z_OK)
+        // chdr.ch_addr_align is not relevant here, because structures will be memcpy'ied out of memory
+        SYMS_U8* uncompressed = syms_push_array(arena, SYMS_U8, uncompressed_size);
+        SYMS_U8* compressed = (SYMS_U8*)file_base + header.sh_offset + chdr_size;
+
+        SYMS_U64 output_size = 0;
+
+        z_stream stream = { 0 };
+        int err = inflateInit(&stream);
+
+        if (err == Z_OK)
+        {
+            stream.next_in = (const Bytef*)compressed;
+            stream.next_out = (Bytef*)uncompressed;
+
+            const SYMS_U64 max_chunk = SYMS_U32_MAX; // max uInt value for zlib
+
+            do
+            {
+                if (stream.avail_out == 0)
+                {
+                    stream.avail_out = (uInt)SYMS_ClampTop(uncompressed_size, max_chunk);
+                    uncompressed_size -= stream.avail_out;
+                }
+                if (stream.avail_in == 0)
+                {
+                    stream.avail_in = (uInt)SYMS_ClampTop(compressed_size, max_chunk);
+                    compressed_size -= stream.avail_in;
+                }
+
+                uInt avail_out = stream.avail_out;
+
+                err = inflate(&stream, Z_NO_FLUSH);
+                if (err == Z_OK || err == Z_STREAM_END)
+                {
+                    output_size += avail_out - stream.avail_out;
+                }
+            }
+            while (err == Z_OK);
+
+            if (err != Z_STREAM_END)
+            {
+                output_size = 0;
+            }
+
+            inflateEnd(&stream);
+        }
+
+        if (output_size == 0)
         {
           // failed to decompress with zlib, this is failure
           result.v = 0;
@@ -260,9 +315,9 @@ syms_elf_section_array_from_img_header(SYMS_Arena *arena, SYMS_String8 file, SYM
 
         // adjust range so when original data pointer is added, we get offsets into newly allocated space
         sections[section_idx].file_range.min = uncompressed - (SYMS_U8*)file_base;
-        sections[section_idx].file_range.max = uncompressed - (SYMS_U8*)file_base + uncompressed_size;
+        sections[section_idx].file_range.max = uncompressed - (SYMS_U8*)file_base + output_size;
       }
-      //else if (chdr64.ch_type == ELF_CompressType_ZStd)
+      //else if (chdr.ch_type == ELF_CompressType_ZStd)
       //{
       //  TODO
       //}
